[SVGImporter]
Move plugin

#virtualized

diff --git a/Engine/Plugins/Experimental/SVGImporter/Source/SVGImporter/Public/ASVGDynamicMeshesContainerActor.h b/Engine/Plugins/Experimental/SVGImporter/Source/SVGImporter/Public/ASVGDynamicMeshesContainerActor.h
new file mode 100644
--- /dev/null
+++ b/Engine/Plugins/Experimental/SVGImporter/Source/SVGImporter/Public/ASVGDynamicMeshesContainerActor.h
@@ -0,0 +1,17 @@
+// Copyright Epic Games, Inc. All Rights Reserved.
+
+#pragma once
+
+#include "GameFramework/Actor.h"
+#include "ASVGDynamicMeshesContainerActor.generated.h"
+
+class UDynamicMeshComponent;
+
+UCLASS(Abstract)
+class SVGIMPORTER_API ASVGDynamicMeshesContainerActor : public AActor
+{
+	GENERATED_BODY()
+public:
+	virtual TArray<UDynamicMeshComponent*> GetSVGDynamicMeshes()
+	PURE_VIRTUAL(ASVGDynamicMeshesContainerActor::GetSVGDynamicMeshes, return {};)
+};
diff --git a/Engine/Plugins/Experimental/SVGImporter/Source/SVGImporter/Public/ProceduralMeshes/SVGFillComponent.h b/Engine/Plugins/Experimental/SVGImporter/Source/SVGImporter/Public/ProceduralMeshes/SVGFillComponent.h
new file mode 100644
--- /dev/null
+++ b/Engine/Plugins/Experimental/SVGImporter/Source/SVGImporter/Public/ProceduralMeshes/SVGFillComponent.h
@@ -0,0 +1,143 @@
+// Copyright Epic Games, Inc. All Rights Reserved.
+
+#pragma once
+
+#include "SVGDynamicMeshComponent.h"
+#include "SVGFillComponent.generated.h"
+
+USTRUCT()
+struct FSVGFillShape
+{
+	GENERATED_BODY()
+
+	UPROPERTY()
+	TArray<FVector2D> ShapeVertices;
+
+	FSVGFillShape(){}
+	FSVGFillShape(const TArray<FVector2D>& InVertices) : ShapeVertices(InVertices){}
+};
+
+USTRUCT()
+struct FSVGFillMeshData
+{
+	GENERATED_BODY()
+
+	UPROPERTY()
+	TArray<FSVGFillShape> ShapesToDraw;
+
+	UPROPERTY()
+	TArray<FSVGFillShape> ShapesToRemove;
+
+	void Init(const TArray<TArray<FVector2D>>& InShapesToDraw, const TArray<TArray<FVector2D>>& InShapesToRemove)
+	{
+		for (const TArray<FVector2D>& Shape : InShapesToDraw)
+		{
+			ShapesToDraw.Emplace(Shape);
+		}
+
+		for (const TArray<FVector2D>& Shape : InShapesToRemove)
+		{
+			ShapesToRemove.Emplace(Shape);
+		}
+	}
+
+	void GetShapesArrays(TArray<TArray<FVector2D>>& OutShapesToDraw, TArray<TArray<FVector2D>>& OutShapesToRemove)
+	{
+		for (const FSVGFillShape& Elem : ShapesToDraw)
+		{
+			OutShapesToDraw.Add(Elem.ShapeVertices);
+		}
+
+		for (const FSVGFillShape& Elem : ShapesToRemove)
+		{
+			OutShapesToRemove.Add(Elem.ShapeVertices);
+		}
+	}
+};
+
+struct FSVGFillParameters
+{
+	const TArray<TArray<FVector2D>>& ShapesToDraw;
+
+	const TArray<TArray<FVector2D>>& ShapesToRemove;
+
+	FColor Color;
+
+	float Extrude = 0;
+
+	bool bSimplify = false;
+
+	float BevelDistance = 0;
+
+	bool bSmoothShapes = false;
+
+	float SmoothingOffset = 150.0;
+
+	bool bUnlit = false;
+
+	bool bCastShadow = false;
+
+	FSVGFillParameters(const TArray<TArray<FVector2D>>& ShapesToDraw, const TArray<TArray<FVector2D>>& ShapesToRemove)
+		: ShapesToDraw(ShapesToDraw),
+		  ShapesToRemove(ShapesToRemove)
+	{
+	}
+};
+
+UCLASS(ClassGroup=(SVGImporter), Meta = (BlueprintSpawnableComponent))
+class SVGIMPORTER_API USVGFillComponent : public USVGDynamicMeshComponent
+{
+	GENERATED_BODY()
+
+public:
+	void GenerateFillMesh(const FSVGFillParameters& InFillParameters);
+	void SetSmoothFillShapes(bool bInSmoothFillShapes, float InSmoothingOffset);
+
+	//~ Begin UObject
+#if WITH_EDITOR
+	virtual void PostTransacted(const FTransactionObjectEvent& TransactionEvent) override;
+	virtual void PostEditImport() override;
+#endif
+	virtual void PostLoad() override;
+	//~ End UObject
+
+protected:
+	//~ Begin USVGDynamicMesh
+	virtual void RegisterDelegates() override;
+	virtual void RegenerateMesh() override;
+	//~ End USVGDynamicMesh
+
+	void RefreshNormals();
+	void CreateSubtractMesh(TWeakObjectPtr<UDynamicMesh>& SubtractMesh, float InExtrude) const;
+	void InitBasePolygons();
+	void GenerateFillMeshInternal();
+
+	void LoadCachedMeshToDynamicMesh(const TOptional<UE::Geometry::FDynamicMesh3>& InSourceMesh);
+	void ApplyExtrudeAndSubtract(float InExtrudeValue);
+	void ApplySimplifyAndBevel();
+
+	/** Applies all changes after extrusion. Includes beveling. */
+	void ApplyFinalMeshChanges();
+
+	/** Applies all changes from extrusion on. */
+	void ApplyFillExtrude();
+
+	/** Applies changes from extrude on. */
+	void ApplyFillBevel();
+
+	TOptional<UE::Geometry::FDynamicMesh3> CachedExtrudedMesh;
+	TOptional<UE::Geometry::FDynamicMesh3> CachedCutPolygon;
+	TOptional<UE::Geometry::FDynamicMesh3> CachedBasePolygon;
+
+	UPROPERTY()
+	bool bSimplify;
+
+	UPROPERTY()
+	bool bBypassClipper;
+
+	UPROPERTY()
+	float SmoothingOffset;
+
+	UPROPERTY()
+	FSVGFillMeshData FillMeshData;
+};
diff --git a/Engine/Plugins/Experimental/SVGImporter/Source/SVGImporter/Public/ProceduralMeshes/SVGStrokeComponent.h b/Engine/Plugins/Experimental/SVGImporter/Source/SVGImporter/Public/ProceduralMeshes/SVGStrokeComponent.h
new file mode 100644
--- /dev/null
+++ b/Engine/Plugins/Experimental/SVGImporter/Source/SVGImporter/Public/ProceduralMeshes/SVGStrokeComponent.h
@@ -0,0 +1,97 @@
+// Copyright Epic Games, Inc. All Rights Reserved.
+
+#pragma once
+
+#include "SVGDynamicMeshComponent.h"
+#include "SVGStrokeComponent.generated.h"
+
+/** @note: the below two enums should match the native ones exactly, so they can be cast directly */
+
+UENUM()
+enum class EPolygonOffsetJoinType : uint8
+{
+	Square, /** Uniform squaring on all convex edge joins. */
+	Round,  /** Arcs on all convex edge joins. */
+	Miter,  /** Squaring of convex edge joins with acute angles ("spikes"). Use in combination with MiterLimit. */
+};
+
+UENUM()
+enum class EPolygonOffsetEndType : uint8
+{
+	Polygon, /** Offsets only one side of a closed path */
+	Joined,  /** Offsets both sides of a path, with joined ends */
+	Butt,    /** Offsets both sides of a path, with square blunt ends */
+	Square,  /** Offsets both sides of a path, with square extended ends */
+	Round,   /** Offsets both sides of a path, with round extended ends */
+};
+
+struct FSVGStrokeParameters
+{
+	const TArray<FVector>& StrokePoints;
+
+	float Thickness = 0;
+
+	FColor InColor;
+
+	bool bIsClosed = false;
+
+	bool bIsClockwise = false;
+
+	EPolygonOffsetJoinType JoinStyle{};
+
+	float Extrude = 0;
+
+	bool bUnlit = false;
+
+	bool bCastShadow = false;
+
+	FSVGStrokeParameters(const TArray<FVector>& InPoints)
+		: StrokePoints(InPoints)
+	{
+	}
+};
+
+UCLASS(ClassGroup=(SVGImporter), Meta = (BlueprintSpawnableComponent))
+class SVGIMPORTER_API USVGStrokeComponent : public USVGDynamicMeshComponent
+{
+	GENERATED_BODY()
+
+public:
+	void GenerateStrokeMesh(const FSVGStrokeParameters& InStrokeParameters);
+
+	void SetJointStyle(EPolygonOffsetJoinType InJoinStyle);
+	void SetStrokeWidth(float InStrokesWidth);
+
+	//~ Begin UObject
+#if WITH_EDITOR
+	virtual void PostTransacted(const FTransactionObjectEvent& TransactionEvent) override;
+#endif
+	//~ End UObject
+
+protected:
+	//~ Begin USVGDynamicMesh
+	virtual void RegisterDelegates() override;
+	virtual void RegenerateMesh() override;
+	//~ End USVGDynamicMesh
+
+	void GenerateStrokeMeshInternal();
+	void GenerateMainStrokeGeometry(float InThickness, float InExtrude);
+	void ApplyStrokeExtrude();
+
+	void CleanupStrokeMesh();
+
+	UPROPERTY()
+	TArray<FVector> StrokePoints;
+
+	UPROPERTY()
+	float StrokeThickness;
+
+	UPROPERTY()
+	bool bStrokeIsClosed;
+
+	UPROPERTY()
+	bool bStrokeIsClockwise;
+
+	UPROPERTY()
+	EPolygonOffsetJoinType JoinStyle;
+};
diff --git a/Engine/Plugins/Experimental/SVGImporter/Source/SVGImporter/Public/SVGGraphicalElements.h b/Engine/Plugins/Experimental/SVGImporter/Source/SVGImporter/Public/SVGGraphicalElements.h
new file mode 100644
--- /dev/null
+++ b/Engine/Plugins/Experimental/SVGImporter/Source/SVGImporter/Public/SVGGraphicalElements.h
@@ -0,0 +1,69 @@
+// Copyright Epic Games, Inc. All Rights Reserved.
+
+#pragma once
+
+#include "SVGPath.h"
+
+/*
+ * A set of graphical elements (lines, rectangles, circles, etc)
+ * implemented as SVG paths (all derive from FSVGPath)
+ */
+
+struct SVGIMPORTER_API FSVGLine : public FSVGPath
+{
+public:
+	FSVGLine(float InPathLength, float InX1, float InY1, float InX2, float InY2);
+	
+	float X1 = 0.0f;
+	float Y1 = 0.0f;
+	float X2 = 0.0f;
+	float Y2 = 0.0f;
+};
+
+struct SVGIMPORTER_API FSVGPolyLine : public FSVGPath
+{
+public:
+	FSVGPolyLine(float InPathLength, const TArray<FVector2D>& InPoints);
+	
+	TArray<FVector2D> Points;
+};
+
+struct SVGIMPORTER_API FSVGPolygon : public FSVGPolyLine
+{
+public:
+	FSVGPolygon(float InPathLength, const TArray<FVector2D>& InPoints);
+};
+
+struct SVGIMPORTER_API FSVGRectangle : public FSVGPath
+{
+public:	
+	FSVGRectangle(float InX, float InY, float InWidth, float InHeight, float InRx, float InRy);
+	
+	float X;
+	float Y;
+	float Width;
+	float Height;
+	float Rx;
+	float Ry;
+};
+
+struct SVGIMPORTER_API FSVGCircle : public FSVGPath
+{
+public:	
+	FSVGCircle(float InCx, float InCy, float InR);
+	
+	float Cx;
+	float Cy;
+	float R;
+};
+
+struct SVGIMPORTER_API FSVGEllipse : public FSVGPath
+{
+public:
+	FSVGEllipse(float InCx, float InCy, float InRx, float InRy);
+	
+	float Cx;
+	float Cy;
+	float Rx;
+	float Ry;
+};
diff --git a/Engine/Plugins/Experimental/SVGImporter/Source/SVGImporterEditor/Private/Factories/SVGActorFactory.cpp b/Engine/Plugins/Experimental/SVGImporter/Source/SVGImporterEditor/Private/Factories/SVGActorFactory.cpp
new file mode 100644
--- /dev/null
+++ b/Engine/Plugins/Experimental/SVGImporter/Source/SVGImporterEditor/Private/Factories/SVGActorFactory.cpp
@@ -0,0 +1,65 @@
+// Copyright Epic Games, Inc. All Rights Reserved.
+
+#include "SVGActorFactory.h"
+
+#include "LevelEditorViewport.h"
+#include "SVGActor.h"
+#include "SVGData.h"
+
+USVGActorFactory::USVGActorFactory(const FObjectInitializer& ObjectInitializer): Super(ObjectInitializer)
+{
+	// AActor subclass this ActorFactory creates.
+	NewActorClassName = FString("ASVGActor");
+	NewActorClass = ASVGActor::StaticClass();
+}
+
+bool USVGActorFactory::CanCreateActorFrom(const FAssetData& AssetData, FText& OutErrorMsg)
+{
+	return (AssetData.IsValid() && AssetData.IsInstanceOf(USVGData::StaticClass()));
+}
+
+AActor* USVGActorFactory::GetDefaultActor(const FAssetData& AssetData)
+{
+	return NewActorClass->GetDefaultObject<ASVGActor>();
+}
+
+AActor* USVGActorFactory::SpawnActor(UObject* InAsset, ULevel* InLevel, const FTransform& InTransform, const FActorSpawnParameters& InSpawnParams)
+{
+	if (!InAsset)
+	{
+		return nullptr;
+	}
+
+	USVGData* SVGData = Cast<USVGData>(InAsset);
+	
+	if (!SVGData)
+	{
+		return nullptr;
+	}
+	
+	const AActor* const DefaultActor = GetDefaultActor(FAssetData(InAsset));
+	
+	if (DefaultActor && InLevel)
+	{
+		const ULevel* const LocalLevel = ValidateSpawnActorLevel(InLevel, InSpawnParams);
+		
+		ASVGActor* const SVGActor = LocalLevel->OwningWorld->SpawnActorDeferred<ASVGActor>(DefaultActor->GetClass(), InTransform);
+		
+		SVGActor->SVGData = SVGData;
+
+		if (FLevelEditorViewportClient::IsDroppingPreviewActor() || InSpawnParams.ObjectFlags & RF_Transient)
+		{
+			SVGActor->RenderMode = ESVGRenderMode::Texture2D;
+		}
+		else
+		{
+			SVGActor->RenderMode = ESVGRenderMode::DynamicMesh3D;
+		}
+		
+		SVGActor->FinishSpawning(InTransform);
+		
+		return SVGActor;
+	}
+
+	return nullptr;
+}
diff --git a/Engine/Plugins/Experimental/SVGImporter/Source/SVGImporterEditor/Private/SVGDynamicMeshVisualizer.cpp b/Engine/Plugins/Experimental/SVGImporter/Source/SVGImporterEditor/Private/SVGDynamicMeshVisualizer.cpp
new file mode 100644
--- /dev/null
+++ b/Engine/Plugins/Experimental/SVGImporter/Source/SVGImporterEditor/Private/SVGDynamicMeshVisualizer.cpp
@@ -0,0 +1,301 @@
+// Copyright Epic Games, Inc. All Rights Reserved.
+
+#include "SVGDynamicMeshVisualizer.h"
+#include "EditorViewportClient.h"
+#include "HitProxies.h"
+#include "ProceduralMeshes/SVGDynamicMeshComponent.h"
+#include "SVGActor.h"
+#include "SceneManagement.h"
+
+float FSVGDynamicMeshVisualizer::FillsExtrudeMin;
+float FSVGDynamicMeshVisualizer::FillsExtrudeMax;
+float FSVGDynamicMeshVisualizer::StrokesExtrudeMin;
+float FSVGDynamicMeshVisualizer::StrokesExtrudeMax;
+
+struct HSVGActorExtrudeHitProxy : HComponentVisProxy
+{
+	DECLARE_HIT_PROXY();
+
+	HSVGActorExtrudeHitProxy(const UActorComponent* InComponent, ASVGActor* InSVGActor, bool bInIsFirstComponent)
+		: HComponentVisProxy(InComponent, HPP_Wireframe)
+		, SVGActorWeak(InSVGActor)
+		, bIsFirstComponent(bInIsFirstComponent)
+	{ }
+
+	TWeakObjectPtr<ASVGActor> SVGActorWeak;
+	bool bIsFirstComponent = false;
+};
+
+IMPLEMENT_HIT_PROXY(HSVGActorExtrudeHitProxy, HComponentVisProxy)
+
+FSVGDynamicMeshVisualizer::FSVGDynamicMeshVisualizer()
+{
+	bIsFirstSVGComponent = false;
+	bIsExtruding = false;
+
+	UpdateMinMaxExtrudeValues();
+}
+
+void FSVGDynamicMeshVisualizer::DrawVisualization(const UActorComponent* InComponent, const FSceneView* InView, FPrimitiveDrawInterface* InPDI)
+{
+	constexpr static float ExtrudeHandleSize = 10.0f;
+
+	const USVGDynamicMeshComponent* SVGDynamicMeshComponent = Cast<USVGDynamicMeshComponent>(InComponent);
+	if (!SVGDynamicMeshComponent)
+	{
+		InPDI->SetHitProxy(nullptr);
+		EndEditing();
+		return;
+	}
+
+	if (ASVGActor* SVGActor = Cast<ASVGActor>(SVGDynamicMeshComponent->GetOwner()))
+	{
+		if (SVGActor->ExtrudeType == ESVGExtrudeType::None)
+		{
+			EndEditing();
+		}
+		else
+		{
+			SVGActorWeak = SVGActor;
+			const TArray<UDynamicMeshComponent*>& SVGShapeComponents = SVGActor->GetSVGDynamicMeshes();
+
+			if (!SVGShapeComponents.IsEmpty())
+			{
+				if (SVGDynamicMeshComponent == SVGShapeComponents[0])
+				{
+					const FVector ProxyHandleLocation = GetExtrudeWidgetLocation();
+					const FVector LineStartLocation = GetExtrudeSurfaceLocation();
+					InPDI->DrawLine(LineStartLocation, ProxyHandleLocation, FLinearColor::White, SDPG_Foreground);
+
+					InPDI->SetHitProxy(new HSVGActorExtrudeHitProxy(SVGDynamicMeshComponent, SVGActor, true));
+					InPDI->DrawPoint(ProxyHandleLocation, FLinearColor::White, ExtrudeHandleSize, SDPG_Foreground);
+					InPDI->SetHitProxy(nullptr);
+				}
+			}
+		}
+	}
+
+	// Draw box to highlight bevel setting while in interactive mode
+	if (SVGDynamicMeshComponent->bIsBevelBeingEdited)
+	{
+		FVector BoxLocation = SVGDynamicMeshComponent->GetComponentLocation() - FVector(SVGDynamicMeshComponent->Bounds.BoxExtent.X, 0, 0) + FVector(SVGDynamicMeshComponent->Bevel/2.0, 0, 0);
+
+		if (SVGDynamicMeshComponent->ExtrudeType == ESVGExtrudeType::FrontFaceOnly)
+		{
+			BoxLocation -= FVector(SVGDynamicMeshComponent->Bounds.BoxExtent.X, 0, 0);
+		}
+
+		FBox Box;
+		Box.Max = BoxLocation + FVector(SVGDynamicMeshComponent->Bevel/2.0f, SVGDynamicMeshComponent->Bounds.BoxExtent.Y, SVGDynamicMeshComponent->Bounds.BoxExtent.Z);
+		Box.Min = BoxLocation - FVector(SVGDynamicMeshComponent->Bevel/2.0f, SVGDynamicMeshComponent->Bounds.BoxExtent.Y, SVGDynamicMeshComponent->Bounds.BoxExtent.Z);
+
+		DrawWireBox(InPDI, Box, FColor::Green, SDPG_Foreground);
+	}
+}
+
+bool FSVGDynamicMeshVisualizer::VisProxyHandleClick(FEditorViewportClient* InViewportClient, HComponentVisProxy* VisProxy, const FViewportClick& Click)
+{
+	EndEditing();
+
+	if (Click.GetKey() == EKeys::LeftMouseButton && VisProxy)
+	{
+		if (VisProxy->IsA(HSVGActorExtrudeHitProxy::StaticGetType()))
+		{
+			if (const HSVGActorExtrudeHitProxy* const ExtrudeProxy = static_cast<HSVGActorExtrudeHitProxy*>(VisProxy))
+			{
+				if (ASVGActor* SVGActor = ExtrudeProxy->SVGActorWeak.Get())
+				{
+					if (ExtrudeProxy->bIsFirstComponent)
+					{
+						GEditor->SelectActor(SVGActor, true, false);
+						bIsExtruding = true;
+						bIsFirstSVGComponent = ExtrudeProxy->bIsFirstComponent;
+						SVGActorWeak = SVGActor;
+						SVGMeshComponent = Cast<USVGDynamicMeshComponent>(const_cast<UActorComponent*>(ExtrudeProxy->Component.Get()));
+						InViewportClient->SetWidgetMode(UE::Widget::WM_None);
+						return true;
+					}
+				}
+			}
+		}
+	}
+
+	return false;
+}
+
+bool FSVGDynamicMeshVisualizer::GetWidgetLocation(const FEditorViewportClient* ViewportClient, FVector& OutLocation) const
+{
+	if (bIsFirstSVGComponent && bIsExtruding)
+	{
+		if (SVGActorWeak.IsValid())
+		{
+			OutLocation = GetExtrudeWidgetLocation();
+			return true;
+		}
+	}
+
+	return false;
+}
+
+bool FSVGDynamicMeshVisualizer::HandleInputDelta(FEditorViewportClient* InViewportClient, FViewport* InViewport, FVector& InDeltaTranslate, FRotator& InDeltaRotate, FVector& InDeltaScale)
+{
+	if (!GetEditedComponent() || !GetEditedComponent()->GetOwner())
+	{
+		EndEditing();
+		return false;
+	}
+
+	if (!InDeltaTranslate.IsZero())
+	{
+		if (bIsFirstSVGComponent)
+		{
+			if (bIsExtruding)
+			{
+				if (ASVGActor* SVGActor = SVGActorWeak.Get())
+				{
+					SVGActor->Modify();
+
+					float FillsExtrudeDepth = SVGActor->GetFillsExtrude() - InDeltaTranslate.X;
+					FillsExtrudeDepth = FMath::Clamp(FillsExtrudeDepth, FillsExtrudeMin, FillsExtrudeMax);
+					SVGActor->SetFillsExtrudeInteractive(FillsExtrudeDepth);
+
+					float StrokesExtrudeDepth = SVGActor->GetStrokesExtrude() - InDeltaTranslate.X;
+					StrokesExtrudeDepth = FMath::Clamp(StrokesExtrudeDepth, StrokesExtrudeMin, StrokesExtrudeMax);
+					SVGActor->SetStrokesExtrudeInteractive(StrokesExtrudeDepth);
+
+					return true;
+				}
+			}
+		}
+	}
+
+	EndEditing();
+	return false;
+}
+
+void FSVGDynamicMeshVisualizer::EndEditing()
+{
+	bIsFirstSVGComponent = false;
+	bIsExtruding = false;
+}
+
+void FSVGDynamicMeshVisualizer::TrackingStopped(FEditorViewportClient* InViewportClient, bool bInDidMove)
+{
+	if (ASVGActor* SVGActor = SVGActorWeak.Get())
+	{
+		SVGActor->Modify();
+
+		const float ExtrudeDepth = FMath::Clamp(SVGActor->GetFillsExtrude(), FillsExtrudeMin, FillsExtrudeMax);
+		SVGActor->SetFillsExtrude(ExtrudeDepth);
+
+		const float StrokesExtrudeDepth = FMath::Clamp(SVGActor->GetStrokesExtrude(), StrokesExtrudeMin, StrokesExtrudeMax);
+		SVGActor->SetStrokesExtrude(StrokesExtrudeDepth);
+	}
+
+	FComponentVisualizer::TrackingStopped(InViewportClient, bInDidMove);
+}
+
+UActorComponent* FSVGDynamicMeshVisualizer::GetEditedComponent() const
+{
+	return SVGMeshComponent.Get();
+}
+
+FVector FSVGDynamicMeshVisualizer::GetExtrudeSurfaceLocation() const
+{
+	if (const ASVGActor* SVGActor = SVGActorWeak.Get())
+	{
+		float ExtrudeDepth = FMath::Max(SVGActor->GetFillsExtrude(), SVGActor->GetStrokesExtrude());
+
+		if (SVGActor->ExtrudeType == ESVGExtrudeType::FrontBackMirror)
+		{
+			ExtrudeDepth *= 0.5f;
+		}
+
+		const FVector SurfaceOffset = -SVGActor->GetActorForwardVector() * ExtrudeDepth;
+
+		return SVGActor->GetActorLocation() + SurfaceOffset;
+	}
+
+	return {};
+}
+
+FVector FSVGDynamicMeshVisualizer::GetExtrudeWidgetLocation() const
+{
+	constexpr static float ExtrudeHandleOffset = 10.0f;
+	const FVector SurfaceLocation = GetExtrudeSurfaceLocation();
+
+	if (const ASVGActor* SVGActor = SVGActorWeak.Get())
+	{
+		const float ExtrudeDepth = FMath::Max(SVGActor->GetFillsExtrude(), SVGActor->GetStrokesExtrude());
+		FVector HandleOffset = SVGActor->GetActorForwardVector() * ExtrudeHandleOffset;
+
+		if (ExtrudeDepth >= 0)
+		{
+			HandleOffset = -HandleOffset;
+		}
+
+		return SurfaceLocation + HandleOffset;
+	}
+
+	return {};
+}
+
+void FSVGDynamicMeshVisualizer::UpdateMinMaxExtrudeValues()
+{
+	static float DefaultExtrudeMin = 0.01f;
+	static float DefaultExtrudeMax = 20.0f;
+
+	// Get Min and Max Extrude values from Properties Metadata, to match what the UI allows
+
+	if (const FProperty* FillsExtrudeProperty = ASVGActor::StaticClass()->FindPropertyByName(TEXT("FillsExtrude")))
+	{
+		if (FillsExtrudeProperty->HasMetaData("UIMin"))
+		{
+			FillsExtrudeMin = FCString::Atof(*FillsExtrudeProperty->GetMetaData("UIMin"));
+		}
+		else
+		{
+			FillsExtrudeMin = DefaultExtrudeMin;
+		}
+
+		if (FillsExtrudeProperty->HasMetaData("UIMax"))
+		{
+			FillsExtrudeMax = FCString::Atof(*FillsExtrudeProperty->GetMetaData("UIMax"));
+		}
+		else
+		{
+			FillsExtrudeMax = DefaultExtrudeMax;
+		}
+	}
+	else
+	{
+		FillsExtrudeMin = DefaultExtrudeMin;
+		FillsExtrudeMax = DefaultExtrudeMax;
+	}
+
+	if (const FProperty* StrokesExtrudeProperty = ASVGActor::StaticClass()->FindPropertyByName(TEXT("StrokesExtrude")))
+	{
+		if (StrokesExtrudeProperty->HasMetaData("UIMin"))
+		{
+			StrokesExtrudeMin = FCString::Atof(*StrokesExtrudeProperty->GetMetaData("UIMin"));
+		}
+		else
+		{
+			StrokesExtrudeMin = DefaultExtrudeMin;
+		}
+
+		if (StrokesExtrudeProperty->HasMetaData("UIMax"))
+		{
+			StrokesExtrudeMax = FCString::Atof(*StrokesExtrudeProperty->GetMetaData("UIMax"));
+		}
+		else
+		{
+			StrokesExtrudeMax = DefaultExtrudeMax;
+		}
+	}
+	else
+	{
+		StrokesExtrudeMin = DefaultExtrudeMin;
+		StrokesExtrudeMax = DefaultExtrudeMax;
+	}
+}
diff --git a/Engine/Plugins/Experimental/SVGImporter/Source/SVGImporterEditor/Private/SVGDynamicMeshVisualizer.h b/Engine/Plugins/Experimental/SVGImporter/Source/SVGImporterEditor/Private/SVGDynamicMeshVisualizer.h
new file mode 100644
--- /dev/null
+++ b/Engine/Plugins/Experimental/SVGImporter/Source/SVGImporterEditor/Private/SVGDynamicMeshVisualizer.h
@@ -0,0 +1,40 @@
+// Copyright Epic Games, Inc. All Rights Reserved.
+
+#pragma once
+
+#include "ComponentVisualizer.h"
+
+class ASVGActor;
+class USVGDynamicMeshComponent;
+
+class FSVGDynamicMeshVisualizer : public FComponentVisualizer
+{
+public:
+	static void UpdateMinMaxExtrudeValues();
+
+	FSVGDynamicMeshVisualizer();
+
+	//~ Begin FComponentVisualizer
+	virtual void DrawVisualization(const UActorComponent* InComponent, const FSceneView* InView, FPrimitiveDrawInterface* InPDI) override;
+	virtual bool VisProxyHandleClick(FEditorViewportClient* InViewportClient, HComponentVisProxy* VisProxy, const FViewportClick& Click) override;
+	virtual bool GetWidgetLocation(const FEditorViewportClient* ViewportClient, FVector& OutLocation) const override;
+	virtual bool HandleInputDelta(FEditorViewportClient* InViewportClient, FViewport* InViewport, FVector& InDeltaTranslate, FRotator& InDeltaRotate, FVector& InDeltaScale) override;
+	virtual void EndEditing() override;
+	virtual void TrackingStopped(FEditorViewportClient* InViewportClient, bool bInDidMove) override;
+	virtual UActorComponent* GetEditedComponent() const override;
+	//~ End FComponentVisualizer
+
+private:
+	FVector GetExtrudeWidgetLocation() const;
+	FVector GetExtrudeSurfaceLocation() const;
+
+	static float FillsExtrudeMin;
+	static float FillsExtrudeMax;
+	static float StrokesExtrudeMin;
+	static float StrokesExtrudeMax;
+
+	TWeakObjectPtr<USVGDynamicMeshComponent> SVGMeshComponent;
+	TWeakObjectPtr<ASVGActor> SVGActorWeak;
+	bool bIsFirstSVGComponent = false;
+	bool bIsExtruding = false;
+};
