[Motion Design] SVGImporter - some css styled SVGs are imported without style/colors

- Fix CSS parser bug which was not considering } as end of attribute value (implicit ';' character)

diff --git a/Engine/Plugins/Experimental/SVGImporter/Source/SVGImporter/Private/SVGImporterUtils.cpp b/Engine/Plugins/Experimental/SVGImporter/Source/SVGImporter/Private/SVGImporterUtils.cpp
--- a/Engine/Plugins/Experimental/SVGImporter/Source/SVGImporter/Private/SVGImporterUtils.cpp
+++ b/Engine/Plugins/Experimental/SVGImporter/Source/SVGImporter/Private/SVGImporterUtils.cpp
@@ -324,8 +324,10 @@ TArray<FSVGStyle> FSVGImporterUtils::StylesFromCSS(FString InData)
 		// Iterate through string.
 		for (int32 i = 0; i < Length; i++)
 		{
+			const TCHAR* CurrentCharacter = Start + i;
+
 			// If we found a {
-			if (FCString::Strncmp(Start + i, TEXT("{"),  1) == 0)
+			if (FCString::Strncmp(CurrentCharacter, TEXT("{"),  1) == 0)
 			{
 				if (CurrClassContentBeginIndex == INDEX_NONE)
 				{
@@ -348,7 +350,7 @@ TArray<FSVGStyle> FSVGImporterUtils::StylesFromCSS(FString InData)
 				}
 			}
 			// If we found a :
-			else if (FCString::Strncmp(Start + i, TEXT(":"),  1) == 0)
+			else if (FCString::Strncmp(CurrentCharacter, TEXT(":"),  1) == 0)
 			{
 				if (CurrAttributeValueBeginIndex == INDEX_NONE)
 				{
@@ -360,8 +362,9 @@ TArray<FSVGStyle> FSVGImporterUtils::StylesFromCSS(FString InData)
 					UE_LOG(LogSVGImporter, Warning, TEXT("Unexpected ':' found while parsing SVG CSS style string."));
 				}
 			}
-			// If we found a ;
-			else if (FCString::Strncmp(Start + i, TEXT(";"),  1) == 0)
+			// If we found a ; or } we can finalize the attribute value
+			else if (FCString::Strncmp(CurrentCharacter, TEXT(";"),  1) == 0
+					|| FCString::Strncmp(CurrentCharacter,  TEXT("}"), 1) == 0)
 			{
 				if (CurrAttributeNameEndIndex != INDEX_NONE && CurrAttributeValueBeginIndex != INDEX_NONE)
 				{
@@ -384,7 +387,9 @@ TArray<FSVGStyle> FSVGImporterUtils::StylesFromCSS(FString InData)
 					UE_LOG(LogSVGImporter, Warning, TEXT("Unexpected ';' found while parsing SVG CSS style string."));
 				}
 			}
-			else if (FCString::Strncmp(Start + i,  TEXT("}"), 1) == 0)
+
+			// If we found a } we can finalize the current style
+			if (FCString::Strncmp(CurrentCharacter,  TEXT("}"), 1) == 0)
 			{
 				if (CurrClassContentBeginIndex != INDEX_NONE)
 				{
