[SVGImporter] Regression: Editor crashing when creating SVG Actor from Clipboard

diff --git a/Engine/Plugins/Experimental/SVGImporter/Source/SVGImporterEditor/Private/SVGImporterEditorModule.cpp b/Engine/Plugins/Experimental/SVGImporter/Source/SVGImporterEditor/Private/SVGImporterEditorModule.cpp
--- a/Engine/Plugins/Experimental/SVGImporter/Source/SVGImporterEditor/Private/SVGImporterEditorModule.cpp
+++ b/Engine/Plugins/Experimental/SVGImporter/Source/SVGImporterEditor/Private/SVGImporterEditorModule.cpp
@@ -259,7 +259,13 @@ void FSVGImporterEditorModule::CreateSVGActorFromClipboard(ULevel* InLevel)
 
 	if (!ClipboardContent.IsEmpty())
 	{
-		USVGData* SVGData = FSVGImporterEditorUtils::CreateSVGDataFromTextBuffer(ClipboardContent, GetTransientPackage(), TEXT("SVGData"), RF_Transactional | RF_Transient);
+		const FName SVGDataUniqueName = MakeUniqueObjectName(GetTransientPackage(), USVGData::StaticClass(), FName(TEXT("SVGData")));
+		USVGData* SVGData = FSVGImporterEditorUtils::CreateSVGDataFromTextBuffer(ClipboardContent, GetTransientPackage(), SVGDataUniqueName, RF_Transactional | RF_Transient);
+		if (!SVGData)
+		{
+			return;
+		}
+
 		if (ASVGActor* NewSVGActor = Cast<ASVGActor>(FActorFactoryAssetProxy::AddActorForAsset(SVGData)))
 		{
 			// We change the object outer so that SVGData belongs to SVGActor
