[SVGImporter] Issue with some "cut" shapes geometry parsing + configuration

SVG shapes used to merge all "cut" holes in a drawing onto a single cut shape, in the Fill Component.
This might lead to wrongly generated geometry.
This CL makes it so the drawn shapes and the "cut" shapes are applied in their appearing order: e.g. Add, Cut, Add, Add, Cut, etc... instead of Add, Add, Add, Cut, Cut.

Additionally, since some SVGs present "cut" shapes before the shape they actually need to carve a hole into, if those "cut" shapes are fully geometrically contained by those shapes to be carved, there's an additional check swapping their position in the shapes list. This allows to properly create their geometry, creating holes which will be otherwise missed.

These solutions are limited by the fact the the initial assumption for these fill rules are thought for 2D rendering.
The nonzero fill rule is in fact executed per pixel, which is not something we can do for 3D geometry, where we have to think about vertices.

This means that, at least with the current implementation, some shapes configuration will not generate the expected result. For example, something like multiple shapes intersecting and/or partially overlapping might not be properly cut.

That would possibly require some complex logic to divide the subshapes into other shapes, based on their overlaps. I think other software trying to address SVGs for 3D is facing some of these issues, since I saw similarly "wrong" geometry generated for the same shapes as with the SVGImporter solution, or generally cut shapes not being automatically handled.

Even though I found those overlaps only in text SVGs, and so far never in "proper" SVGs, e.g. logos, that might still happen, and be worth looking into.

diff --git a/Engine/Plugins/Experimental/SVGImporter/Source/SVGImporter/Private/ProceduralMeshes/SVGFillComponent.cpp b/Engine/Plugins/Experimental/SVGImporter/Source/SVGImporter/Private/ProceduralMeshes/SVGFillComponent.cpp
--- a/Engine/Plugins/Experimental/SVGImporter/Source/SVGImporter/Private/ProceduralMeshes/SVGFillComponent.cpp
+++ b/Engine/Plugins/Experimental/SVGImporter/Source/SVGImporter/Private/ProceduralMeshes/SVGFillComponent.cpp
@@ -13,9 +13,19 @@
 #include "Misc/TransactionObjectEvent.h"
 #endif
 
+void FSVGFillMeshData::Init(const TArray<FSVGPathPolygon>& InShapes)
+{
+	ShapesToDraw.Reset(InShapes.Num());
+	for (const FSVGPathPolygon& Shape : InShapes)
+	{
+		FSVGFillShape FillShape(Shape.Get2DVertices(), Shape.GetShouldBeDrawn());
+		ShapesToDraw.Emplace(FillShape);
+	}
+}
+
 void USVGFillComponent::GenerateFillMesh(const FSVGFillParameters& InFillParameters)
 {
-	FillMeshData.Init(InFillParameters.ShapesToDraw, InFillParameters.ShapesToRemove);
+	FillMeshData.Init(InFillParameters.Shapes);
 
 	SVGColor = InFillParameters.Color;
 	Color = SVGColor;
@@ -39,7 +49,7 @@ void USVGFillComponent::GenerateFillMeshInternal()
 	InitBasePolygons();
 
 	LoadCachedMeshToDynamicMesh(CachedBasePolygon);
-	ApplyExtrudeAndSubtract(GetExtrudeDepth());
+	ApplyExtrude(GetExtrudeDepth());
 	ApplySimplifyAndBevel();
 	ApplyFinalMeshChanges();
 	StoreCurrentMesh();
@@ -55,57 +65,62 @@ void USVGFillComponent::InitBasePolygons()
 	UDynamicMesh* DynamicMesh = NewObject<UDynamicMesh>();
 
 	// main polygon mesh
-	constexpr FGeometryScriptPrimitiveOptions PrimitiveOptions;
 	for (const FSVGFillShape& ShapeToDraw : FillMeshData.ShapesToDraw)
 	{
 		if (ShapeToDraw.ShapeVertices.Num() >= 3)
 		{
-			if (bBypassClipper)
-			{
-				UGeometryScriptLibrary_MeshPrimitiveFunctions::AppendTriangulatedPolygon(DynamicMesh, PrimitiveOptions,  FTransform::Identity, ShapeToDraw.ShapeVertices);
-			}
-			else
-			{
-				UE::Geometry::FGeneralPolygon2d ShapePath{ShapeToDraw.ShapeVertices};
-				TArray<UE::Geometry::FGeneralPolygon2d> OffsetResult;
+			constexpr FGeometryScriptPrimitiveOptions PrimitiveOptions;
 
-				UE::Geometry::PolygonsOffset(SmoothingOffset, {ShapePath}, OffsetResult, true, 1.0, UE::Geometry::EPolygonOffsetJoinType::Round, UE::Geometry::EPolygonOffsetEndType::Polygon);
-
-				if (!OffsetResult.IsEmpty())
+			// Add shape
+			if (ShapeToDraw.bShouldBeDrawn)
+			{
+				if (bBypassClipper)
 				{
-					UGeometryScriptLibrary_MeshPrimitiveFunctions::AppendTriangulatedPolygon(DynamicMesh, PrimitiveOptions, FTransform::Identity, OffsetResult[0].GetOuter().GetVertices());
+					UGeometryScriptLibrary_MeshPrimitiveFunctions::AppendTriangulatedPolygon(DynamicMesh, PrimitiveOptions,  FTransform::Identity, ShapeToDraw.ShapeVertices);
 				}
-			}
-
-			DynamicMesh->ProcessMesh([this](const FDynamicMesh3& Mesh)
-			{
-				CachedBasePolygon = Mesh;
-			});
-		}
-	}
-
-	DynamicMesh->Reset();
-
-	// holes polygon mesh
-	for (const FSVGFillShape& ShapeToRemove : FillMeshData.ShapesToRemove)
-	{
-		if (ShapeToRemove.ShapeVertices.Num() >= 3)
-		{
-			UE::Geometry::FGeneralPolygon2d ShapePath{ShapeToRemove.ShapeVertices};
-			TArray<UE::Geometry::FGeneralPolygon2d> OffsetResult;
+				else
+				{
+					UE::Geometry::FGeneralPolygon2d ShapePath{ShapeToDraw.ShapeVertices};
+					TArray<UE::Geometry::FGeneralPolygon2d> OffsetResult;
 
-			UE::Geometry::PolygonsOffset(SmoothingOffset, {ShapePath}, OffsetResult, true, 1.0, UE::Geometry::EPolygonOffsetJoinType::Round, UE::Geometry::EPolygonOffsetEndType::Polygon);
+					UE::Geometry::PolygonsOffset(SmoothingOffset, {ShapePath}, OffsetResult, true, 1.0, UE::Geometry::EPolygonOffsetJoinType::Round, UE::Geometry::EPolygonOffsetEndType::Polygon);
 
-			if (!OffsetResult.IsEmpty())
+					if (!OffsetResult.IsEmpty())
+					{
+						UGeometryScriptLibrary_MeshPrimitiveFunctions::AppendTriangulatedPolygon(DynamicMesh, PrimitiveOptions, FTransform::Identity, OffsetResult[0].GetOuter().GetVertices());
+					}
+				}
+			}
+			// Cut hole
+			else
 			{
-				UGeometryScriptLibrary_MeshPrimitiveFunctions::AppendTriangulatedPolygon(DynamicMesh, PrimitiveOptions, FTransform::Identity, OffsetResult[0].GetOuter().GetVertices());
+				UDynamicMesh* SubtractMesh = NewObject<UDynamicMesh>();
+				UGeometryScriptLibrary_MeshPrimitiveFunctions::AppendTriangulatedPolygon(SubtractMesh, PrimitiveOptions, FTransform::Identity, ShapeToDraw.ShapeVertices);
+
+				FTransform SubtractMeshTransform = FTransform::Identity;
+				SubtractMeshTransform.SetTranslation(FVector::DownVector * 0.5f);
+				UGeometryScriptLibrary_MeshTransformFunctions::TransformMesh(SubtractMesh, SubtractMeshTransform);
+
+				FGeometryScriptMeshLinearExtrudeOptions ExtrudeOptions;
+				ExtrudeOptions.Distance  = 1.0f;
+				ExtrudeOptions.Direction = FVector::UpVector;
+
+				UGeometryScriptLibrary_MeshModelingFunctions::ApplyMeshLinearExtrudeFaces(SubtractMesh
+					, ExtrudeOptions
+					, FGeometryScriptMeshSelection()
+					, nullptr);
+
+				FGeometryScriptMeshBooleanOptions SubtractOptions;
+				SubtractOptions.bFillHoles = false;
+				SubtractOptions.bSimplifyOutput = true;
+				UGeometryScriptLibrary_MeshBooleanFunctions::ApplyMeshBoolean(DynamicMesh, FTransform::Identity, SubtractMesh, FTransform::Identity, EGeometryScriptBooleanOperation::Subtract, SubtractOptions);
 			}
 		}
 	}
 
 	DynamicMesh->ProcessMesh([this](const FDynamicMesh3& Mesh)
 	{
-		CachedCutPolygon = Mesh;
+		CachedBasePolygon = Mesh;
 	});
 }
 
@@ -166,32 +181,6 @@ void USVGFillComponent::RefreshNormals()
 	UGeometryScriptLibrary_MeshNormalsFunctions::RecomputeNormals(FillMesh, NormalsOptions);
 }
 
-void USVGFillComponent::CreateSubtractMesh(TWeakObjectPtr<UDynamicMesh>& SubtractMesh, float InExtrude) const
-{
-	SubtractMesh = NewObject<UDynamicMesh>();
-
-	if (CachedCutPolygon.IsSet())
-	{
-		SubtractMesh->EditMesh([this](FDynamicMesh3& EditMesh)
-		{
-			EditMesh = CachedCutPolygon.GetValue();
-		});
-
-		FTransform SubtractMeshTransform = FTransform::Identity;
-		SubtractMeshTransform.SetTranslation(FVector::DownVector * InExtrude * 0.5);
-		UGeometryScriptLibrary_MeshTransformFunctions::TransformMesh(SubtractMesh.Get(), SubtractMeshTransform);
-
-		FGeometryScriptMeshLinearExtrudeOptions ExtrudeOptions;
-		ExtrudeOptions.Distance  = InExtrude;
-		ExtrudeOptions.Direction = FVector::UpVector;
-
-		UGeometryScriptLibrary_MeshModelingFunctions::ApplyMeshLinearExtrudeFaces(SubtractMesh.Get()
-			, ExtrudeOptions
-			, FGeometryScriptMeshSelection()
-			, nullptr);
-	}
-}
-
 void USVGFillComponent::LoadCachedMeshToDynamicMesh(const TOptional<UE::Geometry::FDynamicMesh3>& InSourceMesh)
 {
 	if (!InSourceMesh.IsSet())
@@ -212,9 +201,9 @@ void USVGFillComponent::LoadCachedMeshToDynamicMesh(const TOptional<UE::Geometry
 	});
 }
 
-void USVGFillComponent::ApplyExtrudeAndSubtract(float InExtrudeValue)
+void USVGFillComponent::ApplyExtrude(float InExtrudeValue)
 {
-	UDynamicMesh* FillMesh = GetDynamicMesh();
+	const UDynamicMesh* FillMesh = GetDynamicMesh();
 	if (!FillMesh)
 	{
 		return;
@@ -232,18 +221,7 @@ void USVGFillComponent::ApplyExtrudeAndSubtract(float InExtrudeValue)
 			, nullptr);
 	}
 
-	// store it for later use
-	TWeakObjectPtr<UDynamicMesh> SubtractMesh;
-
-	const float SubtractExtrude = FMath::IsNearlyZero(InExtrudeValue) ? 1.0f : InExtrudeValue * 2.0f;
-	CreateSubtractMesh(SubtractMesh, SubtractExtrude);
-
-	FGeometryScriptMeshBooleanOptions SubtractOptions;
-	SubtractOptions.bFillHoles = false;
-	SubtractOptions.bSimplifyOutput = true;
-	UGeometryScriptLibrary_MeshBooleanFunctions::ApplyMeshBoolean(FillMesh, FTransform::Identity, SubtractMesh.Get(), FTransform::Identity, EGeometryScriptBooleanOperation::Subtract, SubtractOptions);
-
-	// store extruded mesh, pre-bevel
+	// Store extruded mesh, pre-bevel
 	FillMesh->ProcessMesh([&](const FDynamicMesh3& SourceEditMesh)
 	{
 		CachedExtrudedMesh = SourceEditMesh;
@@ -279,7 +257,7 @@ void USVGFillComponent::ApplyFillExtrude()
 	}
 
 	LoadCachedMeshToDynamicMesh(CachedBasePolygon);
-	ApplyExtrudeAndSubtract(MinExtrudeValue + Extrude);
+	ApplyExtrude(MinExtrudeValue + Extrude);
 	ApplySimplifyAndBevel();
 	ApplyFinalMeshChanges();
 	StoreCurrentMesh();
diff --git a/Engine/Plugins/Experimental/SVGImporter/Source/SVGImporter/Private/SVGActor.cpp b/Engine/Plugins/Experimental/SVGImporter/Source/SVGImporter/Private/SVGActor.cpp
--- a/Engine/Plugins/Experimental/SVGImporter/Source/SVGImporter/Private/SVGActor.cpp
+++ b/Engine/Plugins/Experimental/SVGImporter/Source/SVGImporter/Private/SVGActor.cpp
@@ -362,8 +362,7 @@ void ASVGActor::AddStrokeComponent(const TArray<FVector>& InPoints, float InThic
 	}
 }
 
-void ASVGActor::AddFillComponent(const TArray<TArray<FVector2D>>& InShapesToDraw, const FColor& InColor, const TArray<TArray<FVector2D>>& InShapesToRemove
-	, float InExtrudeOffset /* = 0 */, const FString& InName /* = "" */)
+void ASVGActor::AddFillComponent(const TArray<FSVGPathPolygon>& InShapesToDraw, const FColor& InColor, float InExtrudeOffset /* = 0 */, const FString& InName /* = "" */)
 {
 	if (!InShapesToDraw.IsEmpty())
 	{
@@ -389,7 +388,7 @@ void ASVGActor::AddFillComponent(const TArray<TArray<FVector2D>>& InShapesToDraw
 
 		DynFillComponent->SetExtrudeType(ExtrudeType);
 
-		FSVGFillParameters FillParams(InShapesToDraw, InShapesToRemove);
+		FSVGFillParameters FillParams(InShapesToDraw);
 		FillParams.Color = InColor;
 		FillParams.Extrude = InExtrudeOffset;
 		FillParams.bSimplify =  UE::SVGImporter::Private::bSimplifyFills;
@@ -455,10 +454,8 @@ void ASVGActor::CreateMeshesFromShape(const FSVGShape& InShape)
 		CurrExtrudeForDepth += DepthExtrude;
 	}
 
-	// each shape is made of a list of vertices
-	TArray<TArray<FVector2D>> FillShapesToCut;
-	TArray<TArray<FVector2D>> FillShapesToRender;
-
+	// Fill shapes need some information other than just vertices in order to be properly generated
+	TArray<FSVGPathPolygon> FillShapesToRender;
 	TArray<TArray<FVector>> StrokesToRender;
 
 	bool bForceFill = false;
@@ -470,28 +467,16 @@ void ASVGActor::CreateMeshesFromShape(const FSVGShape& InShape)
 			StrokesToRender.Add(ShapeToCheck.GetVertices());
 		}
 
-		bool bShouldDrawFill;
+		bool bShouldDrawFill = true;
 
 		if (InShape.PolygonsNum() == 1)
 		{
 			bShouldDrawFill = InShape.HasFill();
 		}
-		else
-		{
-			bShouldDrawFill = ShapeToCheck.GetShouldBeDrawn();
-		}
 
 		if (bShouldDrawFill)
 		{
-			TArray<FVector2D> FillShapeVertices;
-			FillShapeVertices.Append(ShapeToCheck.Get2DVertices());
-			FillShapesToRender.Add(FillShapeVertices);
-		}
-		else
-		{
-			TArray<FVector2D> FillShapeVertices;
-			FillShapeVertices.Append(ShapeToCheck.Get2DVertices());
-			FillShapesToCut.Add(FillShapeVertices);
+			FillShapesToRender.Add(ShapeToCheck);
 		}
 
 		bForceFill = bShouldDrawFill;
@@ -501,7 +486,7 @@ void ASVGActor::CreateMeshesFromShape(const FSVGShape& InShape)
 
 	if (InShape.HasFill() || bForceFill)
 	{
-		AddFillComponent(FillShapesToRender, InShape.GetFillColor(), FillShapesToCut, CurrExtrudeForDepth, InShape.GetId());
+		AddFillComponent(FillShapesToRender, InShape.GetFillColor(), CurrExtrudeForDepth, InShape.GetId());
 		bIncreaseExtrudeForDepth = true;
 	}
 
@@ -1021,6 +1006,12 @@ void ASVGActor::OnShapesGenerationEnd()
 #endif
 
 	ShowShapes();
+
+	// Make sure shapes are updated after generation
+	RefreshFillShapesBevel(EPropertyChangeType::ValueSet);
+	RefreshFillShapesExtrude(EPropertyChangeType::ValueSet);
+	RefreshStrokeShapesExtrude(EPropertyChangeType::ValueSet);
+	ApplyOffset();
 }
 
 void ASVGActor::OnShapesExtrudeEnd()
diff --git a/Engine/Plugins/Experimental/SVGImporter/Source/SVGImporter/Private/SVGImporterUtils.cpp b/Engine/Plugins/Experimental/SVGImporter/Source/SVGImporter/Private/SVGImporterUtils.cpp
--- a/Engine/Plugins/Experimental/SVGImporter/Source/SVGImporter/Private/SVGImporterUtils.cpp
+++ b/Engine/Plugins/Experimental/SVGImporter/Source/SVGImporter/Private/SVGImporterUtils.cpp
@@ -172,6 +172,11 @@ bool FSVGImporterUtils::ShouldPolygonBeDrawn(const FSVGPathPolygon& ShapeToCheck
 	// Let's check the Ray against all other shapes
 	for (const FSVGPathPolygon& Shape : Polygons)
 	{
+		if (Shape.GetPolygon().GetVertices() == ShapeToCheck.GetPolygon().GetVertices())
+		{
+			continue;
+		}
+		
 		TArray<UE::Geometry::TSegment2<double>> Intersections;
 		if (FindIntersectingSegments(Ray, Shape.GetPolygon(), Intersections))
 		{
diff --git a/Engine/Plugins/Experimental/SVGImporter/Source/SVGImporter/Private/SVGTypes.cpp b/Engine/Plugins/Experimental/SVGImporter/Source/SVGImporter/Private/SVGTypes.cpp
--- a/Engine/Plugins/Experimental/SVGImporter/Source/SVGImporter/Private/SVGTypes.cpp
+++ b/Engine/Plugins/Experimental/SVGImporter/Source/SVGImporter/Private/SVGTypes.cpp
@@ -160,6 +160,11 @@ void FSVGShape::AddPolygon(const TArray<FVector>& InPoints)
 
 void FSVGShape::ApplyFillRule()
 {
+	// We currently apply a "polygon" version of the nonzero fill rule.
+	// The way these rules are originally defined for SVGs seems to be reliant on a 2D usage of SVG data.
+	// Our polygon based implementation cannot properly handle certain shape "cut" configurations.
+	// In particular, this applies to overlapping geometry e.g. multiple rectangles intersecting and "cutting" each other.
+
 	for (FSVGPathPolygon& SubShape : Polygons)
 	{
 		bool bShouldDrawFill = false;
@@ -171,7 +176,40 @@ void FSVGShape::ApplyFillRule()
 		{
 			bShouldDrawFill = true;
 		}
-		
+
 		SubShape.SetShouldBeDrawn(bShouldDrawFill);
 	}
+
+	// It may happen that "cut" shapes are defined before the shape they actually cut.
+	// In this case, if the "cut" shape is fully contained by the "cut-destination" shape
+	// we try to swap them, so that the geometry operations later on can properly cut the geometry.
+	// This might not account for 100% of cases (e.g. overlapping shapes, see comment above)
+
+	int32 CurrShapeCount = 0;
+	for (FSVGPathPolygon& SubShape : Polygons)
+	{
+		if (!SubShape.GetShouldBeDrawn())
+		{
+			int32 OtherShapeCount = 0;
+			for (FSVGPathPolygon& OtherShape : Polygons)
+			{
+				if (CurrShapeCount != OtherShapeCount)
+				{
+					if (OtherShape.GetPolygon().Contains(SubShape.GetPolygon().GetVertices()))
+					{
+						if (OtherShapeCount > CurrShapeCount)
+						{
+							// Let's swap the shapes. Like this, the cut shape will be handled after the shape it needs to cut.
+							// The boolean modifier will be able to actually cut the geometry.
+							Polygons.Swap(OtherShapeCount, CurrShapeCount);
+						}
+					}
+				}
+
+				OtherShapeCount++;
+			}
+		}
+
+		CurrShapeCount++;
+	}
 }
diff --git a/Engine/Plugins/Experimental/SVGImporter/Source/SVGImporter/Public/ProceduralMeshes/SVGFillComponent.h b/Engine/Plugins/Experimental/SVGImporter/Source/SVGImporter/Public/ProceduralMeshes/SVGFillComponent.h
--- a/Engine/Plugins/Experimental/SVGImporter/Source/SVGImporter/Public/ProceduralMeshes/SVGFillComponent.h
+++ b/Engine/Plugins/Experimental/SVGImporter/Source/SVGImporter/Public/ProceduralMeshes/SVGFillComponent.h
@@ -3,6 +3,7 @@
 #pragma once
 
 #include "SVGDynamicMeshComponent.h"
+#include "SVGTypes.h"
 #include "SVGFillComponent.generated.h"
 
 USTRUCT()
@@ -13,8 +14,14 @@ struct FSVGFillShape
 	UPROPERTY()
 	TArray<FVector2D> ShapeVertices;
 
+	UPROPERTY()
+	bool bShouldBeDrawn = true;
+
 	FSVGFillShape(){}
-	FSVGFillShape(const TArray<FVector2D>& InVertices) : ShapeVertices(InVertices){}
+	FSVGFillShape(const TArray<FVector2D>& InVertices, bool bInShouldBeDrawn)
+	: ShapeVertices(InVertices)
+	, bShouldBeDrawn(bInShouldBeDrawn)
+	{}
 };
 
 USTRUCT()
@@ -25,61 +32,31 @@ struct FSVGFillMeshData
 	UPROPERTY()
 	TArray<FSVGFillShape> ShapesToDraw;
 
-	UPROPERTY()
-	TArray<FSVGFillShape> ShapesToRemove;
-
-	void Init(const TArray<TArray<FVector2D>>& InShapesToDraw, const TArray<TArray<FVector2D>>& InShapesToRemove)
-	{
-		for (const TArray<FVector2D>& Shape : InShapesToDraw)
-		{
-			ShapesToDraw.Emplace(Shape);
-		}
-
-		for (const TArray<FVector2D>& Shape : InShapesToRemove)
-		{
-			ShapesToRemove.Emplace(Shape);
-		}
-	}
-
-	void GetShapesArrays(TArray<TArray<FVector2D>>& OutShapesToDraw, TArray<TArray<FVector2D>>& OutShapesToRemove)
-	{
-		for (const FSVGFillShape& Elem : ShapesToDraw)
-		{
-			OutShapesToDraw.Add(Elem.ShapeVertices);
-		}
-
-		for (const FSVGFillShape& Elem : ShapesToRemove)
-		{
-			OutShapesToRemove.Add(Elem.ShapeVertices);
-		}
-	}
+	void Init(const TArray<FSVGPathPolygon>& InShapes);
 };
 
 struct FSVGFillParameters
 {
-	const TArray<TArray<FVector2D>>& ShapesToDraw;
-
-	const TArray<TArray<FVector2D>>& ShapesToRemove;
+	TArray<FSVGPathPolygon> Shapes;
 
 	FColor Color;
 
-	float Extrude = 0;
+	float Extrude = 0.0f;
 
 	bool bSimplify = false;
 
-	float BevelDistance = 0;
+	float BevelDistance = 0.0f;
 
 	bool bSmoothShapes = false;
 
-	float SmoothingOffset = 150.0;
+	float SmoothingOffset = 150.0f;
 
 	bool bUnlit = false;
 
 	bool bCastShadow = false;
 
-	FSVGFillParameters(const TArray<TArray<FVector2D>>& ShapesToDraw, const TArray<TArray<FVector2D>>& ShapesToRemove)
-		: ShapesToDraw(ShapesToDraw),
-		  ShapesToRemove(ShapesToRemove)
+	FSVGFillParameters(const TArray<FSVGPathPolygon>& InShapes)
+		:Shapes(InShapes)
 	{
 	}
 };
@@ -108,12 +85,11 @@ protected:
 	//~ End USVGDynamicMesh
 
 	void RefreshNormals();
-	void CreateSubtractMesh(TWeakObjectPtr<UDynamicMesh>& SubtractMesh, float InExtrude) const;
 	void InitBasePolygons();
 	void GenerateFillMeshInternal();
 
 	void LoadCachedMeshToDynamicMesh(const TOptional<UE::Geometry::FDynamicMesh3>& InSourceMesh);
-	void ApplyExtrudeAndSubtract(float InExtrudeValue);
+	void ApplyExtrude(float InExtrudeValue);
 	void ApplySimplifyAndBevel();
 
 	/** Applies all changes after extrusion. Includes beveling. */
@@ -126,7 +102,6 @@ protected:
 	void ApplyFillBevel();
 
 	TOptional<UE::Geometry::FDynamicMesh3> CachedExtrudedMesh;
-	TOptional<UE::Geometry::FDynamicMesh3> CachedCutPolygon;
 	TOptional<UE::Geometry::FDynamicMesh3> CachedBasePolygon;
 
 	UPROPERTY()
diff --git a/Engine/Plugins/Experimental/SVGImporter/Source/SVGImporter/Public/SVGActor.h b/Engine/Plugins/Experimental/SVGImporter/Source/SVGImporter/Public/SVGActor.h
--- a/Engine/Plugins/Experimental/SVGImporter/Source/SVGImporter/Public/SVGActor.h
+++ b/Engine/Plugins/Experimental/SVGImporter/Source/SVGImporter/Public/SVGActor.h
@@ -128,8 +128,7 @@ protected:
 	void AddStrokeComponent(const TArray<FVector>& InPoints, float InThickness, const FColor& InColor, bool bIsClosed, bool bIsClockwise
 		, float InExtrudeOffset = 0.0f, const FString& InName = TEXT(""));
 
-	void AddFillComponent(const TArray<TArray<FVector2D>>& InShapesToDraw, const FColor& InColor, const TArray<TArray<FVector2D>>& InShapesToRemove
-		, float InExtrudeOffset = 0.0f, const FString& InName = TEXT(""));
+	void AddFillComponent(const TArray<FSVGPathPolygon>& InShapesToDraw, const FColor& InColor, float InExtrudeOffset = 0.0f, const FString& InName = TEXT(""));
 
 	void UpdateFillShapesSmoothing();
 	void UpdateFillShapesSmoothingEnable();
