[SVG Importer] Some imported SVGs are not smooth and show polyline segments

- Add quality option for SVG Data import

diff --git a/Engine/Plugins/Experimental/SVGImporter/Source/SVGImporter/Private/SVGData.cpp b/Engine/Plugins/Experimental/SVGImporter/Source/SVGImporter/Private/SVGData.cpp
--- a/Engine/Plugins/Experimental/SVGImporter/Source/SVGImporter/Private/SVGData.cpp
+++ b/Engine/Plugins/Experimental/SVGImporter/Source/SVGImporter/Private/SVGData.cpp
@@ -1,35 +1,79 @@
 // Copyright Epic Games, Inc. All Rights Reserved.
 
+#if WITH_EDITOR
 #include "SVGData.h"
+#include "EditorFramework/AssetImportData.h"
+#include "EditorReimportHandler.h"
+#include "SVGImporterSettings.h"
 #include "SVGImporterUtils.h"
 #include "SVGTypes.h"
 #include "UObject/Package.h"
-
-#if WITH_EDITOR
-#include "EditorFramework/AssetImportData.h"
 #endif
 
+#if WITH_EDITOR
 void USVGData::Initialize(const FSVGDataInitializer& InInitializer)
 {
 	SVGFileContent = InInitializer.SVGTextBuffer;
 	CreateShapes(InInitializer.Elements);
-
-#if WITH_EDITOR
 	SourceFilePath = UAssetImportData::SanitizeImportFilename(InInitializer.SourceFilename, GetOutermost());
 	GenerateSVGTexture();
-#endif
 }
 
-#if WITH_EDITOR
+void USVGData::Reset()
+{
+	SVGFileContent = TEXT("");
+	Shapes.Empty();
+}
+
 void USVGData::GenerateSVGTexture()
 {
 	SVGTexture = FSVGImporterUtils::CreateSVGTexture(SVGFileContent, this);
 }
-#endif
+
+void USVGData::Reimport()
+{
+	constexpr bool bAskForNewFileIfMissing = true;
+	FReimportManager::Instance()->Reimport(this, bAskForNewFileIfMissing);
+
+	OnSVGDataReimport().Broadcast();
+}
+
+void USVGData::PostEditChangeProperty(FPropertyChangedEvent& InPropertyChangedEvent)
+{
+	Super::PostEditChangeProperty(InPropertyChangedEvent);
+
+	const FName MemberName = InPropertyChangedEvent.GetMemberPropertyName();
+
+	if (MemberName == GET_MEMBER_NAME_CHECKED(USVGData, OverrideQuality))
+	{
+		bool bShouldReimport = false;
+		if (InPropertyChangedEvent.ChangeType == EPropertyChangeType::ValueSet)
+		{
+			bShouldReimport = true;
+		}
+		else if (InPropertyChangedEvent.ChangeType == EPropertyChangeType::ToggleEditable)
+		{
+			ESVGSplineConversionQuality ProjectQuality = ESVGSplineConversionQuality::None;
+			if (const USVGImporterSettings* SVGImporterSettings = GetDefault<USVGImporterSettings>())
+			{
+				ProjectQuality = SVGImporterSettings->GetSplineConversionQuality();
+			}
+
+			if (ProjectQuality != OverrideQuality)
+			{
+				bShouldReimport = true;
+			}
+		}
+
+		if (bShouldReimport)
+		{
+			Reimport();
+		}
+	}
+}
 
 void USVGData::CreateShapes(const TArray<TSharedRef<FSVGBaseElement>>& InSVGElements)
 {
-	constexpr float MaxDistanceFromSvgSpline = 1.e-3f;
 	for (TSharedPtr<FSVGBaseElement> Element : InSVGElements)
 	{
 		if (Element->Type != ESVGElementType::Group && Element->Type != ESVGElementType::Other)
@@ -64,10 +108,10 @@ void USVGData::CreateShapes(const TArray<TSharedRef<FSVGBaseElement>>& InSVGElem
 							FString CurrName = Element->Name + FString::FromInt(Count++);
 
 							USplineComponent* SplineComponent = NewObject<USplineComponent>(GetTransientPackage(), *CurrName);
-#if WITH_EDITOR
+
 							// We need to be able to set different arrive and leave tangents per point
 							SplineComponent->bAllowDiscontinuousSpline = true;
-#endif
+
 							int Key = 0;
 							SplineComponent->ClearSplinePoints(true);
 
@@ -99,7 +143,7 @@ void USVGData::CreateShapes(const TArray<TSharedRef<FSVGBaseElement>>& InSVGElem
 							SplineComponent->UpdateSpline();
 
 							TArray<FVector> Points;
-							FSVGImporterUtils::ConvertSVGSplineToPolyLine(SplineComponent, ESplineCoordinateSpace::Local, MaxDistanceFromSvgSpline, Points);
+							FSVGImporterUtils::ConvertSVGSplineToPolyLine(SplineComponent, ESplineCoordinateSpace::Local, GetConversionQualityFactor(), Points);
 
 							SplineComponent->DestroyComponent();
 
@@ -114,3 +158,42 @@ void USVGData::CreateShapes(const TArray<TSharedRef<FSVGBaseElement>>& InSVGElem
 		}
 	}
 }
+
+float USVGData::GetConversionQualityFactor() const
+{
+	ESVGSplineConversionQuality Quality = ESVGSplineConversionQuality::None;
+
+	if (bEnableOverrideQuality)
+	{
+		Quality = OverrideQuality;
+	}
+	else if (const USVGImporterSettings* SVGImporterSettings = GetDefault<USVGImporterSettings>())
+	{
+		Quality = SVGImporterSettings->GetSplineConversionQuality();
+	}
+
+	switch (Quality)
+	{
+		case ESVGSplineConversionQuality::VeryLow:
+			return 1.e-1f;
+
+		case ESVGSplineConversionQuality::Low:
+			return 1.e-2f;
+
+		case ESVGSplineConversionQuality::Normal:
+			return 1.e-3f;
+
+		case ESVGSplineConversionQuality::Increased:
+			return 1.e-4f;
+
+		case ESVGSplineConversionQuality::High:
+			return 1.e-5f;
+
+		case ESVGSplineConversionQuality::VeryHigh:
+			return 1.e-6f;
+
+		default:
+			return 1.e-3f;
+	}
+}
+#endif
diff --git a/Engine/Plugins/Experimental/SVGImporter/Source/SVGImporter/Private/SVGImporterSettings.cpp b/Engine/Plugins/Experimental/SVGImporter/Source/SVGImporter/Private/SVGImporterSettings.cpp
new file mode 100644
--- /dev/null
+++ b/Engine/Plugins/Experimental/SVGImporter/Source/SVGImporter/Private/SVGImporterSettings.cpp
@@ -0,0 +1,9 @@
+// Copyright Epic Games, Inc. All Rights Reserved.
+
+#include "SVGImporterSettings.h"
+
+USVGImporterSettings::USVGImporterSettings()
+{
+	CategoryName = TEXT("Motion Design");
+	SectionName = TEXT("SVG Importer");
+}
diff --git a/Engine/Plugins/Experimental/SVGImporter/Source/SVGImporter/Private/SVGImporterSettings.h b/Engine/Plugins/Experimental/SVGImporter/Source/SVGImporter/Private/SVGImporterSettings.h
new file mode 100644
--- /dev/null
+++ b/Engine/Plugins/Experimental/SVGImporter/Source/SVGImporter/Private/SVGImporterSettings.h
@@ -0,0 +1,25 @@
+// Copyright Epic Games, Inc. All Rights Reserved.
+
+#pragma once
+
+#include "Engine/DeveloperSettings.h"
+#include "SVGData.h"
+#include "SVGImporterSettings.generated.h"
+
+UCLASS(config=Engine, meta=(DisplayName="SVG Importer"))
+class USVGImporterSettings : public UDeveloperSettings
+{
+	GENERATED_BODY()
+
+public:
+	USVGImporterSettings();
+
+	ESVGSplineConversionQuality GetSplineConversionQuality() const { return SplineConversionQuality; }
+
+private:
+	/**
+	 * The quality used to convert SVG Spline Data into poly lines
+	 */
+	UPROPERTY(Config, EditAnywhere, Category="SVG")
+	ESVGSplineConversionQuality SplineConversionQuality = ESVGSplineConversionQuality::VeryHigh;
+};
diff --git a/Engine/Plugins/Experimental/SVGImporter/Source/SVGImporter/Public/SVGData.h b/Engine/Plugins/Experimental/SVGImporter/Source/SVGImporter/Public/SVGData.h
--- a/Engine/Plugins/Experimental/SVGImporter/Source/SVGImporter/Public/SVGData.h
+++ b/Engine/Plugins/Experimental/SVGImporter/Source/SVGImporter/Public/SVGData.h
@@ -16,7 +16,7 @@ struct FSVGDataInitializer
 	FSVGDataInitializer()
 	{
 	}
-	
+
 	FSVGDataInitializer(const FString& InSVGTextBuffer, const FString& InSourceFilename = TEXT(""))
 		: SVGTextBuffer(InSVGTextBuffer)
 		, SourceFilename(InSourceFilename)
@@ -39,31 +39,49 @@ struct FSVGDataInitializer
 	TArray<TSharedRef<FSVGBaseElement>> Elements;
 };
 
-UCLASS(Blueprintable)
+
+/**
+ * Can be used to set the desired fidelity when converting SVG Splines into Polylines
+ */
+UENUM(BlueprintType)
+enum class ESVGSplineConversionQuality : uint8
+{
+	None UMETA(Hidden),
+	VeryLow,
+	Low,
+	Normal,
+	Increased,
+	High,
+	VeryHigh,
+};
+
+UCLASS()
 class SVGIMPORTER_API USVGData : public UObject
 {
 	GENERATED_BODY()
 
 public:
-	void Reset()
-	{
-		SVGFileContent = TEXT("");
-		Shapes.Empty();
-	}
+#if WITH_EDITOR
+	DECLARE_MULTICAST_DELEGATE(FOnSVGDataReimport);
+	FOnSVGDataReimport& OnSVGDataReimport() { return OnSVGDataReimportDelegate; }
 
 	/** Initialize this SVGData with the information provided by the initializer */
 	void Initialize(const FSVGDataInitializer& InInitializer);
 
-#if WITH_EDITOR
+	void Reset();
+
 	/** Generate a texture from the SVG information */
 	void GenerateSVGTexture();
+	void Reimport();
 
 	/** Get the SVG source file path, if available */
 	const FString& GetSourceFilePath() const { return SourceFilePath; }
-#endif
+
+	virtual void PostEditChangeProperty(FPropertyChangedEvent& InPropertyChangedEvent) override;
 
 	/** Create the shapes composing this SVG Data, starting from the information provided by SVG parsing */
 	void CreateShapes(const TArray<TSharedRef<FSVGBaseElement>>& InSVGElements);
+#endif
 
 	UPROPERTY(VisibleAnywhere, BlueprintReadOnly, Category = "SVG Data")
 	TObjectPtr<UTexture2D> SVGTexture;
@@ -77,5 +95,21 @@ public:
 #if WITH_EDITORONLY_DATA
 	UPROPERTY(VisibleAnywhere, Category = "Source Asset")
 	FString SourceFilePath = TEXT("");
+
+	/** Enable quality overriding option */
+	UPROPERTY(EditAnywhere, Category="SVG Data", meta=(InlineEditConditionToggle))
+	bool bEnableOverrideQuality = false;
+
+	/**
+	 * The quality used to convert SVG Spline Data into poly lines
+	 * Changing this value will trigger a re-import (existing geometry will need to be re-generated).
+	 */
+	UPROPERTY(EditAnywhere, Category="SVG Data", meta = (EditCondition = "bEnableOverrideQuality"))
+	ESVGSplineConversionQuality OverrideQuality = ESVGSplineConversionQuality::VeryHigh;
+
+private:
+	float GetConversionQualityFactor() const;
+
+	FOnSVGDataReimport OnSVGDataReimportDelegate;
 #endif
 };
diff --git a/Engine/Plugins/Experimental/SVGImporter/Source/SVGImporter/SVGImporter.Build.cs b/Engine/Plugins/Experimental/SVGImporter/Source/SVGImporter/SVGImporter.Build.cs
--- a/Engine/Plugins/Experimental/SVGImporter/Source/SVGImporter/SVGImporter.Build.cs
+++ b/Engine/Plugins/Experimental/SVGImporter/Source/SVGImporter/SVGImporter.Build.cs
@@ -20,6 +20,7 @@ public class SVGImporter : ModuleRules
 			{
 				"Core",
 				"CoreUObject",
+				"DeveloperSettings",
 				"Engine",
 				"GeometryAlgorithms",
 				"GeometryCore",
