[SVG Importer] ClosePath command sometimes using wrong destination point

When passing multiple points to a MoveTo command, the very last point was used as initial for the sub path.
Correct behavior should be to keep the first one, as this CL addresses.

diff --git a/Engine/Plugins/Experimental/SVGImporter/Source/SVGImporterEditor/Private/SVGParsingUtils.cpp b/Engine/Plugins/Experimental/SVGImporter/Source/SVGImporterEditor/Private/SVGParsingUtils.cpp
--- a/Engine/Plugins/Experimental/SVGImporter/Source/SVGImporterEditor/Private/SVGParsingUtils.cpp
+++ b/Engine/Plugins/Experimental/SVGImporter/Source/SVGImporterEditor/Private/SVGParsingUtils.cpp
@@ -326,9 +326,9 @@ TArray<TArray<FSVGPathCommand>> FSVGParsingUtils::ParseStringAsPathCommands(cons
 					PointsFromString(CommandString, MPoints);
 					const bool bIsRelative = CommandType == 'm';
 
-					PathMoveTo(SubPaths[CurrPathIndex], CursorPos, MPoints, bIsRelative);
-					CursorPos2 = CursorPos;
-					InitialPoint = CursorPos;
+					// We store the initial point of the sub path to properly apply closepath later, if needed
+					InitialPoint = PathMoveTo(SubPaths[CurrPathIndex], CursorPos, MPoints, bIsRelative);
+					CursorPos2 = InitialPoint;
 				}
 			break;
 
@@ -478,9 +478,11 @@ bool FSVGParsingUtils::IsVisible(const TSharedRef<FSVGRawAttribute>& InAttribute
 	return DisplayString != UE::SVGImporter::Public::SVGConstants::None;
 }
 
-void FSVGParsingUtils::PathMoveTo(TArray<FSVGPathCommand>& OutCommands, FVector2D& InOutCursorPos,
+FVector2D FSVGParsingUtils::PathMoveTo(TArray<FSVGPathCommand>& OutCommands, FVector2D& InOutCursorPos,
 	const TArray<FVector2D>& InArguments, bool bInIsRelative)
 {
+	FVector2D InitialPosition = FVector2D::ZeroVector;
+
 	for (int32 i = 0; i < InArguments.Num(); i++)
 	{
 		if (bInIsRelative)
@@ -498,6 +500,7 @@ void FSVGParsingUtils::PathMoveTo(TArray<FSVGPathCommand>& OutCommands, FVector2
 			FSVGPathCommand MoveToCmd(ESVGPathInstructionType::MoveTo);
 			MoveToCmd.PointTo = InOutCursorPos;
 			OutCommands.Add(MoveToCmd);
+			InitialPosition = MoveToCmd.PointTo;
 		}
 		// Following pairs are treated as LineTos
 		else
@@ -507,6 +510,8 @@ void FSVGParsingUtils::PathMoveTo(TArray<FSVGPathCommand>& OutCommands, FVector2
 			OutCommands.Add(LineToCmd);
 		}
 	}
+
+	return InitialPosition;
 }
 
 void FSVGParsingUtils::PathLineTo(TArray<FSVGPathCommand>& OutCommands, FVector2D& InOutCursorPos,
diff --git a/Engine/Plugins/Experimental/SVGImporter/Source/SVGImporterEditor/Private/SVGParsingUtils.h b/Engine/Plugins/Experimental/SVGImporter/Source/SVGImporterEditor/Private/SVGParsingUtils.h
--- a/Engine/Plugins/Experimental/SVGImporter/Source/SVGImporterEditor/Private/SVGParsingUtils.h
+++ b/Engine/Plugins/Experimental/SVGImporter/Source/SVGImporterEditor/Private/SVGParsingUtils.h
@@ -55,8 +55,9 @@ public:
 	 * @param InOutCursorPos cursor position to be used and updated while parsing the arguments
 	 * @param InArguments instruction arguments
 	 * @param bInIsRelative instruction can be relative or absolute
+	 * @return the initial position of the sub path starting with this MoveTo operation. Should be used to later perform a ClosePath if needed.
 	 */
-	static void PathMoveTo(TArray<FSVGPathCommand>& OutCommands, FVector2D& InOutCursorPos, const TArray<FVector2D>& InArguments, bool bInIsRelative);
+	static FVector2D PathMoveTo(TArray<FSVGPathCommand>& OutCommands, FVector2D& InOutCursorPos, const TArray<FVector2D>& InArguments, bool bInIsRelative);
 
 	/**
 	 * Adds a Line To instruction to the provided Commands list.
