[SVG Importer] Move unnecessarily public classes to Private folder

- Fill and Stroke Components
- SVG Shape Actor
- add "IsExperimentalVersion": true to plugin

diff --git a/Engine/Plugins/Experimental/Avalanche/Source/AvalancheSVGEditor/Private/AvaSVGEditorModule.cpp b/Engine/Plugins/Experimental/Avalanche/Source/AvalancheSVGEditor/Private/AvaSVGEditorModule.cpp
--- a/Engine/Plugins/Experimental/Avalanche/Source/AvalancheSVGEditor/Private/AvaSVGEditorModule.cpp
+++ b/Engine/Plugins/Experimental/Avalanche/Source/AvalancheSVGEditor/Private/AvaSVGEditorModule.cpp
@@ -7,7 +7,6 @@
 #include "Modifiers/AvaBevelModifier.h"
 #include "Modifiers/AvaExtrudeModifier.h"
 #include "ProceduralMeshes/SVGDynamicMeshComponent.h"
-#include "ProceduralMeshes/SVGStrokeComponent.h"
 #include "SVGEngineSubsystem.h"
 #include "SVGShapesParentActor.h"
 #include "Subsystems/ActorModifierCoreSubsystem.h"
@@ -101,7 +100,7 @@ void FAvaSVGEditorModule::OnSVGActorSplit(ASVGShapesParentActor* InSVGShapesPare
 					ExtrudeModifier->SetDepth(ExtrudeValue);
 					ExtrudeModifier->SetExtrudeMode(ExtrudeMode);
 
-					if (Shape->IsA<USVGStrokeComponent>())
+					if (Shape->GetShapeType() == TEXT("Stroke"))
 					{
 						ExtrudeModifier->SetCloseBack(false);
 					}
diff --git a/Engine/Plugins/Experimental/SVGImporter/SVGImporter.uplugin b/Engine/Plugins/Experimental/SVGImporter/SVGImporter.uplugin
--- a/Engine/Plugins/Experimental/SVGImporter/SVGImporter.uplugin
+++ b/Engine/Plugins/Experimental/SVGImporter/SVGImporter.uplugin
@@ -12,6 +12,7 @@
 	"SupportURL": "",
 	"CanContainContent": true,
 	"IsBetaVersion": true,
+	"IsExperimentalVersion": true,
 	"Installed": false,
 	"Modules": [
 		{
@@ -31,4 +32,4 @@
 			"Enabled":  true
 		}
 	]
-}
\ No newline at end of file
+}
diff --git a/Engine/Plugins/Experimental/SVGImporter/Source/SVGImporter/Private/ProceduralMeshes/SVGFillComponent.h b/Engine/Plugins/Experimental/SVGImporter/Source/SVGImporter/Private/ProceduralMeshes/SVGFillComponent.h
new file mode 100644
--- /dev/null
+++ b/Engine/Plugins/Experimental/SVGImporter/Source/SVGImporter/Private/ProceduralMeshes/SVGFillComponent.h
@@ -0,0 +1,119 @@
+// Copyright Epic Games, Inc. All Rights Reserved.
+
+#pragma once
+
+#include "ProceduralMeshes/SVGDynamicMeshComponent.h"
+#include "SVGTypes.h"
+#include "SVGFillComponent.generated.h"
+
+USTRUCT()
+struct FSVGFillShape
+{
+	GENERATED_BODY()
+
+	UPROPERTY()
+	TArray<FVector2D> ShapeVertices;
+
+	UPROPERTY()
+	bool bShouldBeDrawn = true;
+
+	FSVGFillShape(){}
+	FSVGFillShape(const TArray<FVector2D>& InVertices, bool bInShouldBeDrawn)
+	: ShapeVertices(InVertices)
+	, bShouldBeDrawn(bInShouldBeDrawn)
+	{}
+};
+
+USTRUCT()
+struct FSVGFillMeshData
+{
+	GENERATED_BODY()
+
+	UPROPERTY()
+	TArray<FSVGFillShape> ShapesToDraw;
+
+	void Init(const TArray<FSVGPathPolygon>& InShapes);
+};
+
+struct FSVGFillParameters
+{
+	TArray<FSVGPathPolygon> Shapes;
+
+	FColor Color;
+
+	float Extrude = 0.0f;
+
+	bool bSimplify = false;
+
+	float BevelDistance = 0.0f;
+
+	bool bSmoothShapes = false;
+
+	float SmoothingOffset = 150.0f;
+
+	bool bUnlit = false;
+
+	bool bCastShadow = false;
+
+	FSVGFillParameters(const TArray<FSVGPathPolygon>& InShapes)
+		:Shapes(InShapes)
+	{
+	}
+};
+
+UCLASS(ClassGroup=(SVGImporter), Meta = (BlueprintSpawnableComponent))
+class USVGFillComponent : public USVGDynamicMeshComponent
+{
+	GENERATED_BODY()
+
+public:
+	void GenerateFillMesh(const FSVGFillParameters& InFillParameters);
+	void SetSmoothFillShapes(bool bInSmoothFillShapes, float InSmoothingOffset);
+
+	//~ Begin UObject
+#if WITH_EDITOR
+	virtual void PostTransacted(const FTransactionObjectEvent& TransactionEvent) override;
+	virtual void PostEditImport() override;
+#endif
+	virtual void PostLoad() override;
+	//~ End UObject
+
+protected:
+	//~ Begin USVGDynamicMesh
+	virtual void RegisterDelegates() override;
+	virtual void RegenerateMesh() override;
+	virtual FName GetShapeType() const override { return TEXT("Fill"); }
+	//~ End USVGDynamicMesh
+
+	void RefreshNormals();
+	void InitBasePolygons();
+	void GenerateFillMeshInternal();
+
+	void LoadCachedMeshToDynamicMesh(const TOptional<UE::Geometry::FDynamicMesh3>& InSourceMesh);
+	void ApplyExtrude(float InExtrudeValue);
+	void ApplySimplifyAndBevel();
+
+	/** Applies all changes after extrusion. Includes beveling. */
+	void ApplyFinalMeshChanges();
+
+	/** Applies all changes from extrusion on. */
+	void ApplyFillExtrude();
+
+	/** Applies changes from extrude on. */
+	void ApplyFillBevel();
+
+	TOptional<UE::Geometry::FDynamicMesh3> CachedExtrudedMesh;
+	TOptional<UE::Geometry::FDynamicMesh3> CachedBasePolygon;
+
+	UPROPERTY()
+	bool bSimplify;
+
+	UPROPERTY()
+	bool bBypassClipper;
+
+	UPROPERTY()
+	float SmoothingOffset;
+
+	UPROPERTY()
+	FSVGFillMeshData FillMeshData;
+};
diff --git a/Engine/Plugins/Experimental/SVGImporter/Source/SVGImporter/Private/ProceduralMeshes/SVGStrokeComponent.h b/Engine/Plugins/Experimental/SVGImporter/Source/SVGImporter/Private/ProceduralMeshes/SVGStrokeComponent.h
new file mode 100644
--- /dev/null
+++ b/Engine/Plugins/Experimental/SVGImporter/Source/SVGImporter/Private/ProceduralMeshes/SVGStrokeComponent.h
@@ -0,0 +1,78 @@
+// Copyright Epic Games, Inc. All Rights Reserved.
+
+#pragma once
+
+#include "ProceduralMeshes/SVGDynamicMeshComponent.h"
+#include "SVGStrokeComponent.generated.h"
+
+struct FSVGStrokeParameters
+{
+	const TArray<FVector>& StrokePoints;
+
+	float Thickness = 0;
+
+	FColor InColor;
+
+	bool bIsClosed = false;
+
+	bool bIsClockwise = false;
+
+	EPolygonOffsetJoinType JoinStyle{};
+
+	float Extrude = 0;
+
+	bool bUnlit = false;
+
+	bool bCastShadow = false;
+
+	FSVGStrokeParameters(const TArray<FVector>& InPoints)
+		: StrokePoints(InPoints)
+	{
+	}
+};
+
+UCLASS(ClassGroup=(SVGImporter), Meta = (BlueprintSpawnableComponent))
+class USVGStrokeComponent : public USVGDynamicMeshComponent
+{
+	GENERATED_BODY()
+
+public:
+	void GenerateStrokeMesh(const FSVGStrokeParameters& InStrokeParameters);
+
+	void SetJointStyle(EPolygonOffsetJoinType InJoinStyle);
+	void SetStrokeWidth(float InStrokesWidth);
+
+	//~ Begin UObject
+#if WITH_EDITOR
+	virtual void PostTransacted(const FTransactionObjectEvent& TransactionEvent) override;
+#endif
+	//~ End UObject
+
+protected:
+	//~ Begin USVGDynamicMesh
+	virtual void RegisterDelegates() override;
+	virtual void RegenerateMesh() override;
+	virtual FName GetShapeType() const override { return TEXT("Stroke"); }
+	//~ End USVGDynamicMesh
+
+	void GenerateStrokeMeshInternal();
+	void GenerateMainStrokeGeometry(float InThickness, float InExtrude);
+	void ApplyStrokeExtrude();
+
+	void CleanupStrokeMesh();
+
+	UPROPERTY()
+	TArray<FVector> StrokePoints;
+
+	UPROPERTY()
+	float StrokeThickness;
+
+	UPROPERTY()
+	bool bStrokeIsClosed;
+
+	UPROPERTY()
+	bool bStrokeIsClockwise;
+
+	UPROPERTY()
+	EPolygonOffsetJoinType JoinStyle;
+};
diff --git a/Engine/Plugins/Experimental/SVGImporter/Source/SVGImporter/Private/SVGShapeActor.h b/Engine/Plugins/Experimental/SVGImporter/Source/SVGImporter/Private/SVGShapeActor.h
new file mode 100644
--- /dev/null
+++ b/Engine/Plugins/Experimental/SVGImporter/Source/SVGImporter/Private/SVGShapeActor.h
@@ -0,0 +1,36 @@
+﻿// Copyright Epic Games, Inc. All Rights Reserved.
+
+#pragma once
+
+#include "ASVGDynamicMeshesContainerActor.h"
+#include "GameFramework/Actor.h"
+#include "SVGShapeActor.generated.h"
+
+class UDynamicMesh;
+class UDynamicMeshComponent;
+class USVGDynamicMeshComponent;
+
+/**
+ * Upon splitting SVG Actors, multiple shape actors will be created.
+ * Each will have their own SVG Shape.
+ */
+UCLASS(ClassGroup = (SVG), meta = (DisplayName = "SVG Shape Actor", ComponentWrapperClass))
+class ASVGShapeActor : public ASVGDynamicMeshesContainerActor
+{
+	GENERATED_BODY()
+
+public:
+	ASVGShapeActor();
+
+	void SetShape(const USVGDynamicMeshComponent* InShape);
+
+	USVGDynamicMeshComponent* GetShape() { return SVGShape; }
+
+	//~ Begin ASVGDynamicMeshesOwnerActor
+	virtual TArray<UDynamicMeshComponent*> GetSVGDynamicMeshes() override;
+	//~ End ASVGDynamicMeshesOwnerActor
+
+protected:
+	UPROPERTY(Category = "SVG", VisibleAnywhere, BlueprintReadOnly)
+	TObjectPtr<USVGDynamicMeshComponent> SVGShape;
+};
diff --git a/Engine/Plugins/Experimental/SVGImporter/Source/SVGImporter/Public/ProceduralMeshes/SVGDynamicMeshComponent.h b/Engine/Plugins/Experimental/SVGImporter/Source/SVGImporter/Public/ProceduralMeshes/SVGDynamicMeshComponent.h
--- a/Engine/Plugins/Experimental/SVGImporter/Source/SVGImporter/Public/ProceduralMeshes/SVGDynamicMeshComponent.h
+++ b/Engine/Plugins/Experimental/SVGImporter/Source/SVGImporter/Public/ProceduralMeshes/SVGDynamicMeshComponent.h
@@ -28,6 +28,16 @@ enum class ESVGMaterialType : uint8
 	Custom
 };
 
+/** @note: the below two enums should match the native ones exactly, so they can be cast directly */
+
+UENUM()
+enum class EPolygonOffsetJoinType : uint8
+{
+	Square, /** Uniform squaring on all convex edge joins. */
+	Round,  /** Arcs on all convex edge joins. */
+	Miter,  /** Squaring of convex edge joins with acute angles ("spikes"). Use in combination with MiterLimit. */
+};
+
 DECLARE_DELEGATE(FSVGMeshActionDelegate)
 
 class UStaticMesh;
@@ -78,6 +88,8 @@ public:
 
 	float GetExtrudeDepth() const { return MinExtrudeValue + Extrude; }
 
+	virtual FName GetShapeType() const { return TEXT("Base"); }
+
 	UMaterialInstanceDynamic* GetMeshMaterialInstance() { return MeshMaterialInstance; }
 
 	UPROPERTY(EditAnywhere, BlueprintReadWrite, Category="SVG")
diff --git a/Engine/Plugins/Experimental/SVGImporter/Source/SVGImporter/Public/ProceduralMeshes/SVGFillComponent.h b/Engine/Plugins/Experimental/SVGImporter/Source/SVGImporter/Public/ProceduralMeshes/SVGFillComponent.h
deleted file mode 100644
diff --git a/Engine/Plugins/Experimental/SVGImporter/Source/SVGImporter/Public/ProceduralMeshes/SVGStrokeComponent.h b/Engine/Plugins/Experimental/SVGImporter/Source/SVGImporter/Public/ProceduralMeshes/SVGStrokeComponent.h
deleted file mode 100644
diff --git a/Engine/Plugins/Experimental/SVGImporter/Source/SVGImporter/Public/SVGActor.h b/Engine/Plugins/Experimental/SVGImporter/Source/SVGImporter/Public/SVGActor.h
--- a/Engine/Plugins/Experimental/SVGImporter/Source/SVGImporter/Public/SVGActor.h
+++ b/Engine/Plugins/Experimental/SVGImporter/Source/SVGImporter/Public/SVGActor.h
@@ -5,7 +5,6 @@
 #include "ASVGDynamicMeshesContainerActor.h"
 #include "GameFramework/Actor.h"
 #include "ProceduralMeshes/SVGDynamicMeshComponent.h"
-#include "ProceduralMeshes/SVGStrokeComponent.h"
 
 #if WITH_EDITOR
 #include "Widgets/Notifications/SNotificationList.h"
diff --git a/Engine/Plugins/Experimental/SVGImporter/Source/SVGImporter/Public/SVGShapeActor.h b/Engine/Plugins/Experimental/SVGImporter/Source/SVGImporter/Public/SVGShapeActor.h
deleted file mode 100644
