[SVG Importer] SVG Actor: Remove ensure on SVGData when creating shapes

- Shapes will not be re-generated when deleting the asset (up to the user to keep them, in case removing the asset was a mistake)
- Actively setting the svg data to null will instead still clear the geometry.

diff --git a/Engine/Plugins/Experimental/SVGImporter/Source/SVGImporter/Private/SVGActor.cpp b/Engine/Plugins/Experimental/SVGImporter/Source/SVGImporter/Private/SVGActor.cpp
--- a/Engine/Plugins/Experimental/SVGImporter/Source/SVGImporter/Private/SVGActor.cpp
+++ b/Engine/Plugins/Experimental/SVGImporter/Source/SVGImporter/Private/SVGActor.cpp
@@ -437,7 +437,7 @@ void ASVGActor::UpdateStrokesWidth()
 
 void ASVGActor::CreateMeshesFromShape(const FSVGShape& InShape)
 {
-	if (!ensureMsgf(SVGData, TEXT("SVG Data is Invalid")))
+	if (!SVGData)
 	{
 		bIsGeneratingMeshes = false;
 #if WITH_EDITOR
@@ -686,11 +686,14 @@ void ASVGActor::ScheduleShapesGeneration()
 		return;
 	}
 
-	if (!ensureMsgf(SVGData, TEXT("SVG Data is Invalid")))
+	if (!SVGData)
 	{
 #if WITH_EDITOR
 		DisplayMissingSVGDataError(TEXT("Cannot generate SVG geometry."));
 #endif
+
+		DestroySVGDynMeshComponents();
+		DestroySVGPlane();
 		return;
 	}
 
@@ -819,7 +822,7 @@ void ASVGActor::TriggerActorDetailsRefresh()
 
 void ASVGActor::GenerateNextMesh()
 {
-	if (!ensureMsgf(SVGData, TEXT("SVG Data is Invalid")))
+	if (!SVGData)
 	{
 		bIsGeneratingMeshes = false;
 		FCoreDelegates::OnEndFrame.Remove(GenMeshDelegateHandle);
@@ -923,7 +926,7 @@ void ASVGActor::ApplyScaleAndCenter()
 
 void ASVGActor::GenerateShapesOnEndFrame()
 {
-	if (!ensureMsgf(SVGData, TEXT("SVG Data is Invalid")))
+	if (!SVGData)
 	{
 		bIsGeneratingMeshes = false;
 		FCoreDelegates::OnEndFrame.Remove(GenMeshDelegateHandle);
@@ -1493,7 +1496,18 @@ void ASVGActor::PostEditChangeProperty(FPropertyChangedEvent& InPropertyChangedE
 	}
 	else if (MemberName == SVGDataName)
 	{
-		RefreshAllShapes();
+		if (InPropertyChangedEvent.ChangeType == EPropertyChangeType::ValueSet)
+		{
+			if (SVGData)
+			{
+				RefreshAllShapes();
+			}
+			else
+			{
+				DestroySVGDynMeshComponents();
+				DestroySVGPlane();
+			}
+		}
 	}
 	else if (MemberName == SmoothFillShapesName)
 	{
