[SVG Importer] Minor parsing improvement + code cleanup

Main items:

- Slightly improved handling of implicitly requested fills
- Remove unused property from SVGImporter module
- Additional validity check for Outer argument when creating a new SVG Texture

diff --git a/Engine/Plugins/Experimental/SVGImporter/Source/SVGImporter/Private/SVGActor.cpp b/Engine/Plugins/Experimental/SVGImporter/Source/SVGImporter/Private/SVGActor.cpp
--- a/Engine/Plugins/Experimental/SVGImporter/Source/SVGImporter/Private/SVGActor.cpp
+++ b/Engine/Plugins/Experimental/SVGImporter/Source/SVGImporter/Private/SVGActor.cpp
@@ -458,8 +458,6 @@ void ASVGActor::CreateMeshesFromShape(const FSVGShape& InShape)
 	TArray<FSVGPathPolygon> FillShapesToRender;
 	TArray<TArray<FVector>> StrokesToRender;
 
-	bool bForceFill = false;
-
 	for (const FSVGPathPolygon& ShapeToCheck : InShape.GetPolygons())
 	{
 		if (InShape.HasStroke())
@@ -478,13 +476,11 @@ void ASVGActor::CreateMeshesFromShape(const FSVGShape& InShape)
 		{
 			FillShapesToRender.Add(ShapeToCheck);
 		}
-
-		bForceFill = bShouldDrawFill;
 	}
 
 	bool bIncreaseExtrudeForDepth = false;
 
-	if (InShape.HasFill() || bForceFill)
+	if (InShape.HasFill())
 	{
 		AddFillComponent(FillShapesToRender, InShape.GetFillColor(), CurrExtrudeForDepth, InShape.GetId());
 		bIncreaseExtrudeForDepth = true;
diff --git a/Engine/Plugins/Experimental/SVGImporter/Source/SVGImporter/Private/SVGImporterUtils.cpp b/Engine/Plugins/Experimental/SVGImporter/Source/SVGImporter/Private/SVGImporterUtils.cpp
--- a/Engine/Plugins/Experimental/SVGImporter/Source/SVGImporter/Private/SVGImporterUtils.cpp
+++ b/Engine/Plugins/Experimental/SVGImporter/Source/SVGImporter/Private/SVGImporterUtils.cpp
@@ -339,15 +339,18 @@ TArray<FSVGStyle> FSVGImporterUtils::StylesFromCSS(FString InData)
 					CurrClassContentBeginIndex = i + 1;
 
 					const int32 CurrClassNameEndIndex = i;
-
-					// todo: expand on this kind of syntax?  something.class
-					CurrClassName = FString(CurrClassNameEndIndex - CurrClassNameBeginIndex, Start + CurrClassNameBeginIndex);
-					if (CurrClassName.StartsWith("."))
+					const int32 ClassNameLength = CurrClassNameEndIndex - CurrClassNameBeginIndex;
+					if (ClassNameLength >= 0)
 					{
-						CurrClassName.RemoveAt(0);
-					}
+						// todo: expand on this kind of syntax?  something.class
+						CurrClassName = FString(ClassNameLength, Start + CurrClassNameBeginIndex);
+						if (CurrClassName.StartsWith("."))
+						{
+							CurrClassName.RemoveAt(0);
+						}
 
-					CurrAttributeNameBeginIndex = i + 1;
+						CurrAttributeNameBeginIndex = i + 1;
+					}
 				}
 				else
 				{
@@ -375,17 +378,26 @@ TArray<FSVGStyle> FSVGImporterUtils::StylesFromCSS(FString InData)
 				{
 					const int32 CurrAttributeValueEndIndex = i;
 
-					FString AttributeName = FString(CurrAttributeNameEndIndex - CurrAttributeNameBeginIndex, Start + CurrAttributeNameBeginIndex);
-					AttributeName.RemoveSpacesInline();
+					const int32 AttributeNameLength = CurrAttributeNameEndIndex - CurrAttributeNameBeginIndex;
+					const int32 AttributeValueLength = CurrAttributeValueEndIndex - CurrAttributeValueBeginIndex;
+					if (AttributeNameLength >= 0 || AttributeValueLength >= 0)
+					{
+						FString AttributeName = FString(AttributeNameLength, Start + CurrAttributeNameBeginIndex);
+						AttributeName.RemoveSpacesInline();
 
-					FString AttributeValue = FString(CurrAttributeValueEndIndex - CurrAttributeValueBeginIndex, Start + CurrAttributeValueBeginIndex);
-					AttributeValue.RemoveSpacesInline();
+						FString AttributeValue = FString(CurrAttributeValueEndIndex - CurrAttributeValueBeginIndex, Start + CurrAttributeValueBeginIndex);
+						AttributeValue.RemoveSpacesInline();
 
-					CurrStyleAttributes.Add(AttributeName, AttributeValue);
+						CurrStyleAttributes.Add(AttributeName, AttributeValue);
 
-					CurrAttributeNameBeginIndex = i + 1;
-					CurrAttributeNameEndIndex = INDEX_NONE;
-					CurrAttributeValueBeginIndex = INDEX_NONE;
+						CurrAttributeNameBeginIndex = i + 1;
+						CurrAttributeNameEndIndex = INDEX_NONE;
+						CurrAttributeValueBeginIndex = INDEX_NONE;
+					}
+					else
+					{
+						UE_LOG(LogSVGImporter, Warning, TEXT("Error parsing attribute, skipping."));
+					}
 				}
 				else
 				{
@@ -540,13 +552,13 @@ bool FSVGImporterUtils::GetAvailableFolderPath(const FString& InAssetPath, FStri
 
 UTexture2D* FSVGImporterUtils::CreateSVGTexture(const FString& InSVGString, UObject* InOuter)
 {
-	if (InSVGString.IsEmpty())
+	if (InSVGString.IsEmpty() || !InOuter)
 	{
 		return nullptr;
 	}
 
 	// Bytes per pixel (RGBA)
-	constexpr static int32 SVG_BPP = 4;	
+	constexpr static int32 SVG_BPP = 4;
 	constexpr static int32 SVG_DPI = 96;
 
 	NSVGimage* Image = nsvgParse(TCHAR_TO_ANSI(*InSVGString), "px", SVG_DPI);
diff --git a/Engine/Plugins/Experimental/SVGImporter/Source/SVGImporter/Public/SVGImporter.h b/Engine/Plugins/Experimental/SVGImporter/Source/SVGImporter/Public/SVGImporter.h
--- a/Engine/Plugins/Experimental/SVGImporter/Source/SVGImporter/Public/SVGImporter.h
+++ b/Engine/Plugins/Experimental/SVGImporter/Source/SVGImporter/Public/SVGImporter.h
@@ -30,6 +30,4 @@ public:
 protected:
 	DECLARE_DELEGATE_RetVal(USVGData*, FOnDefaultSVGDataRequested)
 	FOnDefaultSVGDataRequested OnDefaultSVGDataRequested;
-
-	TObjectPtr<USVGData> DefaultSVGData;
 };
diff --git a/Engine/Plugins/Experimental/SVGImporter/Source/SVGImporter/Public/SVGTypes.h b/Engine/Plugins/Experimental/SVGImporter/Source/SVGImporter/Public/SVGTypes.h
--- a/Engine/Plugins/Experimental/SVGImporter/Source/SVGImporter/Public/SVGTypes.h
+++ b/Engine/Plugins/Experimental/SVGImporter/Source/SVGImporter/Public/SVGTypes.h
@@ -39,7 +39,7 @@ struct SVGIMPORTER_API FSVGStyle
 
 	bool HasStroke() const { return bHasStroke; }
 
-	bool HasFill() const { return bHasFill; }
+	bool HasFill() const { return bHasFill || !bHasStroke; }
 
 	const FColor& GetFillColor() const { return FillColor; }
 
@@ -58,7 +58,7 @@ private:
 	FString Name;
 
 	UPROPERTY()
-	bool bHasFill = true;
+	bool bHasFill = false;
 
 	UPROPERTY()
 	FColor FillColor = FColor::Black;
diff --git a/Engine/Plugins/Experimental/SVGImporter/Source/SVGImporterEditor/Private/SVGImporterEditorUtils.cpp b/Engine/Plugins/Experimental/SVGImporter/Source/SVGImporterEditor/Private/SVGImporterEditorUtils.cpp
--- a/Engine/Plugins/Experimental/SVGImporter/Source/SVGImporterEditor/Private/SVGImporterEditorUtils.cpp
+++ b/Engine/Plugins/Experimental/SVGImporter/Source/SVGImporterEditor/Private/SVGImporterEditorUtils.cpp
@@ -213,9 +213,9 @@ void FSVGImporterEditorUtils::ParseSVGElement(const TSharedRef<FSVGRawElement>&
 			if (const TSharedPtr<FSVGRawAttribute> AttributeStyle = GetAttribute(SVGConstants::Style, AttributesMap))
 			{
 				// todo: handle this case in the style parser itself, so we don't have to add brackets here
-				const FString& StyleString = TEXT("{") + AttributeStyle->AsString().TrimStartAndEnd() + TEXT("}");
+				FString StyleString = TEXT("{") + AttributeStyle->AsString().TrimStartAndEnd() + TEXT("}");
 
-				TArray<FSVGStyle> Styles = FSVGImporterUtils::StylesFromCSS(StyleString);
+				TArray<FSVGStyle> Styles = FSVGImporterUtils::StylesFromCSS(MoveTemp(StyleString));
 
 				// Really, this should be just one
 				if (Styles.Num() > 0) 
