[SVG Import] Newly created SVG Data asset cannot be properly deleted if object using it is selected

diff --git a/Engine/Plugins/Experimental/SVGImporter/Source/SVGImporter/Private/SVGImporterUtils.cpp b/Engine/Plugins/Experimental/SVGImporter/Source/SVGImporter/Private/SVGImporterUtils.cpp
--- a/Engine/Plugins/Experimental/SVGImporter/Source/SVGImporter/Private/SVGImporterUtils.cpp
+++ b/Engine/Plugins/Experimental/SVGImporter/Source/SVGImporter/Private/SVGImporterUtils.cpp
@@ -591,7 +591,7 @@ UTexture2D* FSVGImporterUtils::CreateSVGTexture(const FString& InSVGString, UObj
 	const FString BaseTextureName = FString(TEXT("SVGTexture"));
 	const FName TextureName = MakeUniqueObjectName(InOuter, UTexture2D::StaticClass(), FName(*BaseTextureName));
 
-	UTexture2D* SVGTexture = NewObject<UTexture2D>(InOuter, UTexture2D::StaticClass(), TextureName, InOuter->GetFlags());
+	UTexture2D* SVGTexture = NewObject<UTexture2D>(InOuter, UTexture2D::StaticClass(), TextureName, RF_Public | RF_Transactional);
 
 	if (!SVGTexture)
 	{
