SVG Importer:
- Fix crash when dragging SVG Actor from Placement Browser

#changelist validated
#virtualized

diff --git a/Engine/Plugins/Experimental/SVGImporter/Source/SVGImporter/Private/SVGActor.cpp b/Engine/Plugins/Experimental/SVGImporter/Source/SVGImporter/Private/SVGActor.cpp
--- a/Engine/Plugins/Experimental/SVGImporter/Source/SVGImporter/Private/SVGActor.cpp
+++ b/Engine/Plugins/Experimental/SVGImporter/Source/SVGImporter/Private/SVGActor.cpp
@@ -62,10 +62,25 @@ namespace UE::SVGImporter::Private
 
 	/** Used to ensure strokes are always slightly in front of fills within the same shape */
 	constexpr static float StrokeAddedDepthMultiplier = 1.01f;
+
+	static bool bAllowInitialization = true; 
 }
 
 #define LOCTEXT_NAMESPACE "SVGActor"
 
+FSVGActorInitGuard::FSVGActorInitGuard()
+{
+	using namespace UE::SVGImporter;
+
+	bPreviousValue = Private::bAllowInitialization;
+	Private::bAllowInitialization = false;
+}
+
+FSVGActorInitGuard::~FSVGActorInitGuard()
+{
+	UE::SVGImporter::Private::bAllowInitialization = bPreviousValue;
+}
+
 // Sets default values
 ASVGActor::ASVGActor()
 {
@@ -121,6 +136,29 @@ ASVGActor::ASVGActor()
 	LoadResources();
 }
 
+void ASVGActor::Initialize()
+{
+	using namespace UE::SVGImporter;
+
+	if (Private::bAllowInitialization && !bInitialized)
+	{
+		ApplyScaleAndCenter();
+		ApplyOffset();
+
+		TryLoadDefaultSVGData();
+
+		if (SVGData)
+		{
+			if (bMeshesShouldBeGenerated)
+			{
+				Generate();
+			}
+		}
+
+		bInitialized = true;
+	}
+}
+
 void ASVGActor::TryLoadDefaultSVGData()
 {
 	if (!SVGData)
@@ -248,24 +286,7 @@ void ASVGActor::PostDuplicate(bool bDuplicateForPIE)
 void ASVGActor::PostRegisterAllComponents()
 {
 	Super::PostRegisterAllComponents();
-
-	if (!bInitialized)
-	{
-		ApplyScaleAndCenter();
-		ApplyOffset();
-
-		TryLoadDefaultSVGData();
-
-		if (SVGData)
-		{
-			if (bMeshesShouldBeGenerated)
-			{
-				Generate();
-			}
-		}
-
-		bInitialized = true;
-	}
+	Initialize();
 }
 
 void ASVGActor::Generate()
@@ -1411,6 +1432,11 @@ void ASVGActor::BakeToBlueprint() const
 
 void ASVGActor::DisplayMissingSVGDataError(const FString& InErrorMsg)
 {
+	if (bIsEditorPreviewActor)
+	{
+		return;
+	}
+
 	// in case this notification is still hanging around, let's remove it
 	HideGenerationStartNotification();
 
diff --git a/Engine/Plugins/Experimental/SVGImporter/Source/SVGImporter/Public/SVGActor.h b/Engine/Plugins/Experimental/SVGImporter/Source/SVGImporter/Public/SVGActor.h
--- a/Engine/Plugins/Experimental/SVGImporter/Source/SVGImporter/Public/SVGActor.h
+++ b/Engine/Plugins/Experimental/SVGImporter/Source/SVGImporter/Public/SVGActor.h
@@ -30,6 +30,15 @@ enum class ESVGRenderMode : uint8
 	Texture2D = 1       UMETA(DisplayName = "2D")
 };
 
+struct SVGIMPORTER_API FSVGActorInitGuard
+{
+	FSVGActorInitGuard();
+	~FSVGActorInitGuard();
+
+private:
+	bool bPreviousValue;
+};
+
 UCLASS()
 class SVGIMPORTER_API ASVGActor : public ASVGDynamicMeshesContainerActor
 {
@@ -38,6 +47,8 @@ class SVGIMPORTER_API ASVGActor : public ASVGDynamicMeshesContainerActor
 public:
 	ASVGActor();
 
+	void Initialize();
+
 	//~ Begin UObject
 	virtual void PostLoad() override;
 	virtual void PostDuplicate(bool bDuplicateForPIE) override;
diff --git a/Engine/Plugins/Experimental/SVGImporter/Source/SVGImporterEditor/Private/Factories/SVGActorFactory.cpp b/Engine/Plugins/Experimental/SVGImporter/Source/SVGImporterEditor/Private/Factories/SVGActorFactory.cpp
--- a/Engine/Plugins/Experimental/SVGImporter/Source/SVGImporterEditor/Private/Factories/SVGActorFactory.cpp
+++ b/Engine/Plugins/Experimental/SVGImporter/Source/SVGImporterEditor/Private/Factories/SVGActorFactory.cpp
@@ -1,15 +1,13 @@
 // Copyright Epic Games, Inc. All Rights Reserved.
 
-#include "SVGActorFactory.h"
-
+#include "Factories/SVGActorFactory.h"
 #include "LevelEditorViewport.h"
 #include "SVGActor.h"
 #include "SVGData.h"
 
-USVGActorFactory::USVGActorFactory(const FObjectInitializer& ObjectInitializer): Super(ObjectInitializer)
+USVGActorFactory::USVGActorFactory(const FObjectInitializer& ObjectInitializer)
+	: Super(ObjectInitializer)
 {
-	// AActor subclass this ActorFactory creates.
-	NewActorClassName = FString("ASVGActor");
 	NewActorClass = ASVGActor::StaticClass();
 }
 
@@ -25,41 +23,23 @@ AActor* USVGActorFactory::GetDefaultActor(const FAssetData& AssetData)
 
 AActor* USVGActorFactory::SpawnActor(UObject* InAsset, ULevel* InLevel, const FTransform& InTransform, const FActorSpawnParameters& InSpawnParams)
 {
-	if (!InAsset)
+	ASVGActor* SVGActor;
 	{
-		return nullptr;
+		FSVGActorInitGuard InitGuard;
+		SVGActor = Cast<ASVGActor>(Super::SpawnActor(InAsset, InLevel, InTransform, InSpawnParams));
 	}
 
-	USVGData* SVGData = Cast<USVGData>(InAsset);
-	
-	if (!SVGData)
+	SVGActor->SVGData = Cast<USVGData>(InAsset);
+
+	if (InSpawnParams.bTemporaryEditorActor)
 	{
-		return nullptr;
+		SVGActor->RenderMode = ESVGRenderMode::Texture2D;
 	}
-	
-	const AActor* const DefaultActor = GetDefaultActor(FAssetData(InAsset));
-	
-	if (DefaultActor && InLevel)
+	else
 	{
-		const ULevel* const LocalLevel = ValidateSpawnActorLevel(InLevel, InSpawnParams);
-		
-		ASVGActor* const SVGActor = LocalLevel->OwningWorld->SpawnActorDeferred<ASVGActor>(DefaultActor->GetClass(), InTransform);
-		
-		SVGActor->SVGData = SVGData;
-
-		if (FLevelEditorViewportClient::IsDroppingPreviewActor() || InSpawnParams.ObjectFlags & RF_Transient)
-		{
-			SVGActor->RenderMode = ESVGRenderMode::Texture2D;
-		}
-		else
-		{
-			SVGActor->RenderMode = ESVGRenderMode::DynamicMesh3D;
-		}
-		
-		SVGActor->FinishSpawning(InTransform);
-		
-		return SVGActor;
+		SVGActor->RenderMode = ESVGRenderMode::DynamicMesh3D;
 	}
 
-	return nullptr;
+	SVGActor->Initialize();
+	return SVGActor;
 }
diff --git a/Engine/Plugins/Experimental/SVGImporter/Source/SVGImporterEditor/Private/Factories/SVGActorFactory.h b/Engine/Plugins/Experimental/SVGImporter/Source/SVGImporterEditor/Private/Factories/SVGActorFactory.h
deleted file mode 100644
diff --git a/Engine/Plugins/Experimental/SVGImporter/Source/SVGImporterEditor/Public/Factories/SVGActorFactory.h b/Engine/Plugins/Experimental/SVGImporter/Source/SVGImporterEditor/Public/Factories/SVGActorFactory.h
new file mode 100644
--- /dev/null
+++ b/Engine/Plugins/Experimental/SVGImporter/Source/SVGImporterEditor/Public/Factories/SVGActorFactory.h
@@ -0,0 +1,21 @@
+// Copyright Epic Games, Inc. All Rights Reserved.
+
+#pragma once
+
+#include "ActorFactories/ActorFactory.h"
+#include "SVGActorFactory.generated.h"
+
+UCLASS(MinimalAPI)
+class USVGActorFactory : public UActorFactory
+{
+	GENERATED_BODY()
+
+	USVGActorFactory(const FObjectInitializer& ObjectInitializer);
+
+protected:
+	//~ Begin UActorFactory
+	virtual bool CanCreateActorFrom(const FAssetData& AssetData, FText& OutErrorMsg) override;
+	virtual AActor* GetDefaultActor(const FAssetData& AssetData) override;
+	virtual AActor* SpawnActor(UObject* InAsset, ULevel* InLevel, const FTransform& InTransform, const FActorSpawnParameters& InSpawnParams) override;
+	//~ End UActorFactory
+};
