SVG Importer:
 - More API cleanup

Unshelved from pending changelist '31512338':

diff --git a/Engine/Plugins/Experimental/SVGImporter/Source/SVGImporter/Private/SVGJoinedShapesActor.h b/Engine/Plugins/Experimental/SVGImporter/Source/SVGImporter/Private/SVGJoinedShapesActor.h
--- a/Engine/Plugins/Experimental/SVGImporter/Source/SVGImporter/Private/SVGJoinedShapesActor.h
+++ b/Engine/Plugins/Experimental/SVGImporter/Source/SVGImporter/Private/SVGJoinedShapesActor.h
@@ -2,7 +2,7 @@
 
 #pragma once
 
-#include "ASVGDynamicMeshesContainerActor.h"
+#include "SVGDynamicMeshesContainerActor.h"
 #include "GameFramework/Actor.h"
 #include "SVGJoinedShapesActor.generated.h"
 
diff --git a/Engine/Plugins/Experimental/SVGImporter/Source/SVGImporter/Private/SVGShapeActor.h b/Engine/Plugins/Experimental/SVGImporter/Source/SVGImporter/Private/SVGShapeActor.h
--- a/Engine/Plugins/Experimental/SVGImporter/Source/SVGImporter/Private/SVGShapeActor.h
+++ b/Engine/Plugins/Experimental/SVGImporter/Source/SVGImporter/Private/SVGShapeActor.h
@@ -2,7 +2,7 @@
 
 #pragma once
 
-#include "ASVGDynamicMeshesContainerActor.h"
+#include "SVGDynamicMeshesContainerActor.h"
 #include "GameFramework/Actor.h"
 #include "SVGShapeActor.generated.h"
 
diff --git a/Engine/Plugins/Experimental/SVGImporter/Source/SVGImporter/Public/ASVGDynamicMeshesContainerActor.h b/Engine/Plugins/Experimental/SVGImporter/Source/SVGImporter/Public/ASVGDynamicMeshesContainerActor.h
deleted file mode 100644
diff --git a/Engine/Plugins/Experimental/SVGImporter/Source/SVGImporter/Public/ProceduralMeshes/JoinedSVGDynamicMeshComponent.h b/Engine/Plugins/Experimental/SVGImporter/Source/SVGImporter/Public/ProceduralMeshes/JoinedSVGDynamicMeshComponent.h
--- a/Engine/Plugins/Experimental/SVGImporter/Source/SVGImporter/Public/ProceduralMeshes/JoinedSVGDynamicMeshComponent.h
+++ b/Engine/Plugins/Experimental/SVGImporter/Source/SVGImporter/Public/ProceduralMeshes/JoinedSVGDynamicMeshComponent.h
@@ -65,7 +65,7 @@ enum class EJoinedSVGMeshColoring : uint8
 	SingleColor
 };
 
-UCLASS(ClassGroup=(Custom), meta=(BlueprintSpawnableComponent))
+UCLASS(MinimalAPI, ClassGroup=(Custom), meta=(BlueprintSpawnableComponent))
 class UJoinedSVGDynamicMeshComponent : public UDynamicMeshComponent
 {
 	friend class FJoinedSVGDynamicMeshComponentCustomization;
diff --git a/Engine/Plugins/Experimental/SVGImporter/Source/SVGImporter/Public/ProceduralMeshes/SVGDynamicMeshComponent.h b/Engine/Plugins/Experimental/SVGImporter/Source/SVGImporter/Public/ProceduralMeshes/SVGDynamicMeshComponent.h
--- a/Engine/Plugins/Experimental/SVGImporter/Source/SVGImporter/Public/ProceduralMeshes/SVGDynamicMeshComponent.h
+++ b/Engine/Plugins/Experimental/SVGImporter/Source/SVGImporter/Public/ProceduralMeshes/SVGDynamicMeshComponent.h
@@ -42,14 +42,17 @@ DECLARE_DELEGATE(FSVGMeshActionDelegate)
 
 class UStaticMesh;
 
-UCLASS(ClassGroup=(SVG))
-class SVGIMPORTER_API USVGDynamicMeshComponent : public UDynamicMeshComponent
+UCLASS(MinimalAPI, ClassGroup=(SVG))
+class USVGDynamicMeshComponent : public UDynamicMeshComponent
 {
 	GENERATED_BODY()
 
 public:
 	USVGDynamicMeshComponent();
 
+	/** Calling this function will flatten the SVG Shape, converting it to a 2D polygon */
+	SVGIMPORTER_API void FlattenShape();
+
 	//~ Begin UObject
 #if WITH_EDITOR
 	virtual void PostEditChangeProperty(FPropertyChangedEvent& PropertyChangedEvent) override;
@@ -61,9 +64,6 @@ public:
 	virtual void PostDuplicate(bool bDuplicateForPIE) override;
 	//~ End UObject
 
-	/** Calling this function will flatten the SVG Shape, converting it to a 2D polygon */
-	void FlattenShape();
-
 	/** Calling this function will scale the SVG Shape */
 	void ScaleShape(float InScale);
 
diff --git a/Engine/Plugins/Experimental/SVGImporter/Source/SVGImporter/Public/SVGActor.h b/Engine/Plugins/Experimental/SVGImporter/Source/SVGImporter/Public/SVGActor.h
--- a/Engine/Plugins/Experimental/SVGImporter/Source/SVGImporter/Public/SVGActor.h
+++ b/Engine/Plugins/Experimental/SVGImporter/Source/SVGImporter/Public/SVGActor.h
@@ -2,7 +2,7 @@
 
 #pragma once
 
-#include "ASVGDynamicMeshesContainerActor.h"
+#include "SVGDynamicMeshesContainerActor.h"
 #include "GameFramework/Actor.h"
 #include "ProceduralMeshes/SVGDynamicMeshComponent.h"
 
diff --git a/Engine/Plugins/Experimental/SVGImporter/Source/SVGImporter/Public/SVGBakedActor.h b/Engine/Plugins/Experimental/SVGImporter/Source/SVGImporter/Public/SVGBakedActor.h
--- a/Engine/Plugins/Experimental/SVGImporter/Source/SVGImporter/Public/SVGBakedActor.h
+++ b/Engine/Plugins/Experimental/SVGImporter/Source/SVGImporter/Public/SVGBakedActor.h
@@ -5,11 +5,11 @@
 #include "GameFramework/Actor.h"
 #include "SVGBakedActor.generated.h"
 
-USTRUCT(Blueprintable)
+USTRUCT()
 struct FSVGBakeElement
 {
 	GENERATED_BODY()
-	
+
 	UPROPERTY(VisibleAnywhere, Category = "SVG Bake Element")
 	TObjectPtr<UStaticMesh> Mesh = nullptr;
 
@@ -20,14 +20,14 @@ struct FSVGBakeElement
 	FString Name;
 };
 
-UCLASS(Blueprintable)
-class SVGIMPORTER_API ASVGBakedActor : public AActor
+UCLASS(MinimalAPI)
+class ASVGBakedActor : public AActor
 {
 	GENERATED_BODY()
 
 public:
 	ASVGBakedActor();
 
-	UPROPERTY(EditAnywhere, BlueprintReadWrite, Category = "Baked SVG")
+	UPROPERTY(EditAnywhere, Category = "Baked SVG")
 	TArray<FSVGBakeElement> SVGBakeElements;
 };
diff --git a/Engine/Plugins/Experimental/SVGImporter/Source/SVGImporter/Public/SVGData.h b/Engine/Plugins/Experimental/SVGImporter/Source/SVGImporter/Public/SVGData.h
--- a/Engine/Plugins/Experimental/SVGImporter/Source/SVGImporter/Public/SVGData.h
+++ b/Engine/Plugins/Experimental/SVGImporter/Source/SVGImporter/Public/SVGData.h
@@ -39,7 +39,6 @@ struct FSVGDataInitializer
 	TArray<TSharedRef<FSVGBaseElement>> Elements;
 };
 
-
 /**
  * Can be used to set the desired fidelity when converting SVG Splines into Polylines
  */
@@ -83,13 +82,13 @@ public:
 	void CreateShapes(const TArray<TSharedRef<FSVGBaseElement>>& InSVGElements);
 #endif
 
-	UPROPERTY(VisibleAnywhere, BlueprintReadOnly, Category = "SVG Data")
+	UPROPERTY(VisibleAnywhere, Category = "SVG Data")
 	TObjectPtr<UTexture2D> SVGTexture;
 
-	UPROPERTY(VisibleAnywhere, BlueprintReadOnly, Category = "SVG Data")
+	UPROPERTY(VisibleAnywhere, Category = "SVG Data")
 	FString SVGFileContent;
 
-	UPROPERTY(BlueprintReadOnly, Category = "SVG Data")
+	UPROPERTY()
 	TArray<FSVGShape> Shapes;
 
 #if WITH_EDITORONLY_DATA
diff --git a/Engine/Plugins/Experimental/SVGImporter/Source/SVGImporter/Public/SVGDynamicMeshesContainerActor.h b/Engine/Plugins/Experimental/SVGImporter/Source/SVGImporter/Public/SVGDynamicMeshesContainerActor.h
new file mode 100644
--- /dev/null
+++ b/Engine/Plugins/Experimental/SVGImporter/Source/SVGImporter/Public/SVGDynamicMeshesContainerActor.h
@@ -0,0 +1,17 @@
+// Copyright Epic Games, Inc. All Rights Reserved.
+
+#pragma once
+
+#include "GameFramework/Actor.h"
+#include "SVGDynamicMeshesContainerActor.generated.h"
+
+class UDynamicMeshComponent;
+
+UCLASS(MinimalAPI, Abstract)
+class ASVGDynamicMeshesContainerActor : public AActor
+{
+	GENERATED_BODY()
+
+public:
+	virtual TArray<UDynamicMeshComponent*> GetSVGDynamicMeshes() PURE_VIRTUAL(ASVGDynamicMeshesContainerActor::GetSVGDynamicMeshes, return {};)
+};
diff --git a/Engine/Plugins/Experimental/SVGImporter/Source/SVGImporter/Public/SVGElements.h b/Engine/Plugins/Experimental/SVGImporter/Source/SVGImporter/Public/SVGElements.h
--- a/Engine/Plugins/Experimental/SVGImporter/Source/SVGImporter/Public/SVGElements.h
+++ b/Engine/Plugins/Experimental/SVGImporter/Source/SVGImporter/Public/SVGElements.h
@@ -74,24 +74,22 @@ struct SVGIMPORTER_API FSVGMatrix
 	bool bIsInitialized = false;
 };
 
-struct SVGIMPORTER_API FSVGRadialGradient
+struct FSVGBaseElement : TSharedFromThis<FSVGBaseElement>
 {
-	FSVGRadialGradient(){}
-};
-
-struct SVGIMPORTER_API FSVGBaseElement : public TSharedFromThis<FSVGBaseElement>
-{
-public:
-	FSVGBaseElement(){}
-
 	void SetName(const FString& InName)
 	{
 		Name = InName;
 	}
 
-	bool TypeIsSet() const { return Type != ESVGElementType::None; }
+	bool TypeIsSet() const
+	{
+		return Type != ESVGElementType::None;
+	}
 
-	bool IsGraphicElement() const { return bIsGraphicElement; }
+	bool IsGraphicElement() const
+	{
+		return bIsGraphicElement;
+	}
 
 	ESVGElementType Type = ESVGElementType::None;
 
@@ -100,18 +98,16 @@ public:
 	FString Name;
 };
 
-struct SVGIMPORTER_API FSVGMainElement : public FSVGBaseElement
+struct FSVGMainElement : FSVGBaseElement
 {
-public:
 	FSVGMainElement()
 	{
 		Type = ESVGElementType::SVG;
 	}
 };
 
-struct SVGIMPORTER_API FSVGStyleElement : public FSVGBaseElement
+struct FSVGStyleElement : FSVGBaseElement
 {
-public:
 	FSVGStyleElement(const TArray<FSVGStyle>& InStyles)
 		: Styles(InStyles)
 	{
@@ -146,29 +142,31 @@ public:
 	}
 };
 
-struct SVGIMPORTER_API FSVGClipPath : public FSVGBaseElement
+struct FSVGClipPath : FSVGBaseElement
 {
-public:
 	FSVGClipPath()
 	{
 		Type = ESVGElementType::ClipPath;
 	}
-
-	/** todo */
 };
 
-struct SVGIMPORTER_API FSVGGradientElement : public FSVGBaseElement
+struct FSVGGradientElement : FSVGBaseElement
 {
-public:
 	FSVGGradientElement(const FSVGGradient& InGradient)
 		: Gradient(InGradient)
 	{
 		Type = ESVGElementType::Gradient;
 	}
 
-	bool HasID() const { return Gradient.IsValid(); }
+	bool HasID() const
+	{
+		return Gradient.IsValid();
+	}
 
-	const FString& GetID() const { return Gradient.Id; }
+	const FString& GetID() const
+	{
+		return Gradient.Id;
+	}
 
 	FSVGGradient Gradient;
 };
@@ -177,9 +175,8 @@ public:
  * Base for all elements which can be drawn.
  * They can have fill and/or stroke information, a style, and a transform.
  */
-struct SVGIMPORTER_API FSVGGraphicsElement : public FSVGBaseElement
+struct SVGIMPORTER_API FSVGGraphicsElement : FSVGBaseElement
 {
-public:
 	FSVGGraphicsElement();
 
 	/** If the specified parent is a Group element, add this Graphic Element to its children */
@@ -341,9 +338,8 @@ protected:
  * It holds a list of graphics elements.
  * It is defined as a graphic element itself, since it can contain info about elements styling.
  */
-struct SVGIMPORTER_API FSVGGroupElement : public FSVGGraphicsElement
+struct FSVGGroupElement : FSVGGraphicsElement
 {
-public:
 	TArray<TSharedRef<FSVGGraphicsElement>> Children;
 
 	FSVGGroupElement()
diff --git a/Engine/Plugins/Experimental/SVGImporter/Source/SVGImporter/Public/SVGEngineSubsystem.h b/Engine/Plugins/Experimental/SVGImporter/Source/SVGImporter/Public/SVGEngineSubsystem.h
--- a/Engine/Plugins/Experimental/SVGImporter/Source/SVGImporter/Public/SVGEngineSubsystem.h
+++ b/Engine/Plugins/Experimental/SVGImporter/Source/SVGImporter/Public/SVGEngineSubsystem.h
@@ -12,7 +12,7 @@ DECLARE_DELEGATE_OneParam(FSVGActorComponentsReady, ASVGActor*)
 DECLARE_DELEGATE_OneParam(FOnSVGActorSplit, ASVGShapesParentActor*)
 
 UCLASS()
-class SVGIMPORTER_API USVGEngineSubsystem : public UEngineSubsystem
+class USVGEngineSubsystem : public UEngineSubsystem
 {
 	GENERATED_BODY()
 
@@ -24,7 +24,7 @@ public:
 	static FOnSVGActorSplit& OnSVGActorSplit() { return OnSVGActorSplitDelegate; }
 
 protected:
-	static FOnSVGActorSplit OnSVGActorSplitDelegate;
+	SVGIMPORTER_API static FOnSVGActorSplit OnSVGActorSplitDelegate;
 
 	FSVGActorComponentsReady SVGActorComponentsReady;
 };
diff --git a/Engine/Plugins/Experimental/SVGImporter/Source/SVGImporter/Public/SVGGraphicalElements.h b/Engine/Plugins/Experimental/SVGImporter/Source/SVGImporter/Public/SVGGraphicalElements.h
--- a/Engine/Plugins/Experimental/SVGImporter/Source/SVGImporter/Public/SVGGraphicalElements.h
+++ b/Engine/Plugins/Experimental/SVGImporter/Source/SVGImporter/Public/SVGGraphicalElements.h
@@ -9,36 +9,32 @@
  * implemented as SVG paths (all derive from FSVGPath)
  */
 
-struct SVGIMPORTER_API FSVGLine : public FSVGPath
+struct FSVGLine : FSVGPath
 {
-public:
-	FSVGLine(float InPathLength, float InX1, float InY1, float InX2, float InY2);
-	
+	SVGIMPORTER_API FSVGLine(float InPathLength, float InX1, float InY1, float InX2, float InY2);
+
 	float X1 = 0.0f;
 	float Y1 = 0.0f;
 	float X2 = 0.0f;
 	float Y2 = 0.0f;
 };
 
-struct SVGIMPORTER_API FSVGPolyLine : public FSVGPath
+struct FSVGPolyLine : FSVGPath
 {
-public:
-	FSVGPolyLine(float InPathLength, const TArray<FVector2D>& InPoints);
-	
+	SVGIMPORTER_API FSVGPolyLine(float InPathLength, const TArray<FVector2D>& InPoints);
+
 	TArray<FVector2D> Points;
 };
 
-struct SVGIMPORTER_API FSVGPolygon : public FSVGPolyLine
+struct FSVGPolygon : FSVGPolyLine
 {
-public:
-	FSVGPolygon(float InPathLength, const TArray<FVector2D>& InPoints);
+	SVGIMPORTER_API FSVGPolygon(float InPathLength, const TArray<FVector2D>& InPoints);
 };
 
-struct SVGIMPORTER_API FSVGRectangle : public FSVGPath
+struct FSVGRectangle : FSVGPath
 {
-public:	
-	FSVGRectangle(float InX, float InY, float InWidth, float InHeight, float InRx, float InRy);
-	
+	SVGIMPORTER_API FSVGRectangle(float InX, float InY, float InWidth, float InHeight, float InRx, float InRy);
+
 	float X;
 	float Y;
 	float Width;
@@ -47,21 +43,19 @@ public:
 	float Ry;
 };
 
-struct SVGIMPORTER_API FSVGCircle : public FSVGPath
-{
-public:	
-	FSVGCircle(float InCx, float InCy, float InR);
-	
+struct FSVGCircle : FSVGPath
+{	
+	SVGIMPORTER_API FSVGCircle(float InCx, float InCy, float InR);
+
 	float Cx;
 	float Cy;
 	float R;
 };
 
-struct SVGIMPORTER_API FSVGEllipse : public FSVGPath
+struct FSVGEllipse : FSVGPath
 {
-public:
-	FSVGEllipse(float InCx, float InCy, float InRx, float InRy);
-	
+	SVGIMPORTER_API FSVGEllipse(float InCx, float InCy, float InRx, float InRy);
+
 	float Cx;
 	float Cy;
 	float Rx;
diff --git a/Engine/Plugins/Experimental/SVGImporter/Source/SVGImporter/Public/SVGImporterUtils.h b/Engine/Plugins/Experimental/SVGImporter/Source/SVGImporter/Public/SVGImporterUtils.h
--- a/Engine/Plugins/Experimental/SVGImporter/Source/SVGImporter/Public/SVGImporterUtils.h
+++ b/Engine/Plugins/Experimental/SVGImporter/Source/SVGImporter/Public/SVGImporterUtils.h
@@ -30,7 +30,7 @@ namespace ESplineCoordinateSpace
 	enum Type : int;
 }
 
-class SVGIMPORTER_API FSVGImporterUtils
+class FSVGImporterUtils
 {
 public:
 	inline static double SVGScaleFactor = 0.1;
@@ -66,10 +66,10 @@ public:
 	static void RotateAroundCustomPivot(FVector& OutPointToRotate, const FVector& InPivotPosition, float InAngle);
 
 	/** Extracts a color from SVG color string */
-	static FColor GetColorFromSVGString(FString InColorString);
+	SVGIMPORTER_API static FColor GetColorFromSVGString(FString InColorString);
 
 	/** Extracts SVG Style from a CSS style string */
-	static TArray<FSVGStyle> StylesFromCSS(FString InString);
+	SVGIMPORTER_API static TArray<FSVGStyle> StylesFromCSS(FString InString);
 
 	/** Sets an SVG Matrix based on the specified SVG Transform string */
 	static bool SetSVGMatrixFromTransformString(const FString& InTransformString, FSVGMatrix& OutSVGMatrix);
@@ -122,7 +122,7 @@ public:
 	 * @param InSVGActor the source SVG Actor to split
 	 * @return the newly created Shapes Parent Actor
 	 */
-	static ASVGShapesParentActor* SplitSVGActor(ASVGActor* InSVGActor);
+	SVGIMPORTER_API static ASVGShapesParentActor* SplitSVGActor(ASVGActor* InSVGActor);
 
 	/**
 	 * Merges all SVG Shapes of this Actor into a single Actor with a single Shape.
@@ -130,7 +130,7 @@ public:
 	 * @param InSVGActor the source SVG Actor to consolidate
 	 * @return the newly created Joined Shapes Actor
 	 */
-	static ASVGJoinedShapesActor* ConsolidateSVGActor(ASVGActor* InSVGActor);
+	SVGIMPORTER_API static ASVGJoinedShapesActor* ConsolidateSVGActor(ASVGActor* InSVGActor);
 
 	/**
 	 * Joins the SVG Shapes of 2 or more Actors of type ASVGDynamicMeshesOwnerActor onto a single ASVGJoinedShapesActor
@@ -138,7 +138,7 @@ public:
 	 * @param InDynamicSVGShapesOwners the list of Actors owning the SVG Shapes which are going to be joined
 	 * @return the newly created Joined Shapes Actor
 	 */
-	static ASVGJoinedShapesActor* JoinSVGDynamicMeshOwners(const TArray<ASVGDynamicMeshesContainerActor*>& InDynamicSVGShapesOwners);
+	SVGIMPORTER_API static ASVGJoinedShapesActor* JoinSVGDynamicMeshOwners(const TArray<ASVGDynamicMeshesContainerActor*>& InDynamicSVGShapesOwners);
 
 private:
 	/** Based on the specified Asset Path, will return an available Folder Path where we can create assets for SVG Baking purposes */
diff --git a/Engine/Plugins/Experimental/SVGImporter/Source/SVGImporter/Public/SVGPath.h b/Engine/Plugins/Experimental/SVGImporter/Source/SVGImporter/Public/SVGPath.h
--- a/Engine/Plugins/Experimental/SVGImporter/Source/SVGImporter/Public/SVGPath.h
+++ b/Engine/Plugins/Experimental/SVGImporter/Source/SVGImporter/Public/SVGPath.h
@@ -90,17 +90,20 @@ public:
 /**
  * SVG Subpath
  */
-struct SVGIMPORTER_API FSVGSubPath
+struct FSVGSubPath
 {
 	/** Set the path closed flag*/
-	void SetIsClosed(bool bInIsClosed) { bIsClosed = bInIsClosed; }
+	void SetIsClosed(bool bInIsClosed)
+	{
+		bIsClosed = bInIsClosed;
+	}
 
 	/** Add an element to this path */
 	void AddPathElement(const FSVGPathElement& InElement)
 	{
 		Elements.Add(InElement);
 	}
-	
+
 	/** Elements of this Sub Path */
 	TArray<FSVGPathElement> Elements;
 
@@ -117,9 +120,9 @@ public:
 	FSVGPath();
 
 	FSVGPath(float InPathLength);
-	
+
 	FSVGPath(const TArray<TArray<FSVGPathCommand>>& InPaths, float InPathLength);
-	
+
 	/** Set the leave tangent of the previous point */
 	bool UpdatePreviousElementLeaveTangent(const FVector2D& InLeaveTangent);
 
diff --git a/Engine/Plugins/Experimental/SVGImporter/Source/SVGImporter/Public/SVGShapesParentActor.h b/Engine/Plugins/Experimental/SVGImporter/Source/SVGImporter/Public/SVGShapesParentActor.h
--- a/Engine/Plugins/Experimental/SVGImporter/Source/SVGImporter/Public/SVGShapesParentActor.h
+++ b/Engine/Plugins/Experimental/SVGImporter/Source/SVGImporter/Public/SVGShapesParentActor.h
@@ -2,7 +2,7 @@
 
 #pragma once
 
-#include "ASVGDynamicMeshesContainerActor.h"
+#include "SVGDynamicMeshesContainerActor.h"
 #include "GameFramework/Actor.h"
 #include "SVGShapesParentActor.generated.h"
 
diff --git a/Engine/Plugins/Experimental/SVGImporter/Source/SVGImporter/Public/SVGTypes.h b/Engine/Plugins/Experimental/SVGImporter/Source/SVGImporter/Public/SVGTypes.h
--- a/Engine/Plugins/Experimental/SVGImporter/Source/SVGImporter/Public/SVGTypes.h
+++ b/Engine/Plugins/Experimental/SVGImporter/Source/SVGImporter/Public/SVGTypes.h
@@ -14,7 +14,7 @@
 
 class USplineComponent;
 
-USTRUCT(BlueprintType)
+USTRUCT()
 struct SVGIMPORTER_API FSVGStyle
 {
 	GENERATED_BODY()
@@ -83,7 +83,7 @@ enum class ESVGGradientPointUnit : uint8
 };
 
 USTRUCT()
-struct SVGIMPORTER_API FSVGGradientStop
+struct FSVGGradientStop
 {
 	GENERATED_BODY()
 
@@ -98,7 +98,7 @@ struct SVGIMPORTER_API FSVGGradientStop
 };
 
 USTRUCT()
-struct SVGIMPORTER_API FSVGGradient
+struct FSVGGradient
 {
 	GENERATED_BODY()
 
@@ -138,7 +138,7 @@ struct SVGIMPORTER_API FSVGGradient
 	//- gradientTransform
 	//
 
-	FColor GetAverageColor() const;
+	SVGIMPORTER_API FColor GetAverageColor() const;
 };
 
 USTRUCT(Blueprintable)
@@ -189,7 +189,7 @@ private:
 	bool bShouldBeDrawn;
 };
 
-USTRUCT(Blueprintable)
+USTRUCT()
 struct FSVGShape
 {
 	GENERATED_BODY()
diff --git a/Engine/Plugins/Experimental/SVGImporter/Source/SVGImporterEditor/Private/AssetDefinition_SVGData.h b/Engine/Plugins/Experimental/SVGImporter/Source/SVGImporterEditor/Private/AssetDefinition_SVGData.h
--- a/Engine/Plugins/Experimental/SVGImporter/Source/SVGImporterEditor/Private/AssetDefinition_SVGData.h
+++ b/Engine/Plugins/Experimental/SVGImporter/Source/SVGImporterEditor/Private/AssetDefinition_SVGData.h
@@ -8,14 +8,14 @@
 struct FToolMenuContext;
 
 UCLASS()
-class SVGIMPORTEREDITOR_API UAssetDefinition_SVGData : public UAssetDefinitionDefault
+class UAssetDefinition_SVGData : public UAssetDefinitionDefault
 {
 	GENERATED_BODY()
 
 public:
 	/** Menu Extension statics */
 	static void ExecuteReimportSVG(const FToolMenuContext& InContext);
-	
+
 	//~ Begin UAssetDefinition
 	virtual FText GetAssetDisplayName() const override;
 	virtual TSoftClassPtr<UObject> GetAssetClass() const override;
diff --git a/Engine/Plugins/Experimental/SVGImporter/Source/SVGImporterEditor/Private/Factories/SVGFactory.h b/Engine/Plugins/Experimental/SVGImporter/Source/SVGImporterEditor/Private/Factories/SVGFactory.h
--- a/Engine/Plugins/Experimental/SVGImporter/Source/SVGImporterEditor/Private/Factories/SVGFactory.h
+++ b/Engine/Plugins/Experimental/SVGImporter/Source/SVGImporterEditor/Private/Factories/SVGFactory.h
@@ -6,7 +6,7 @@
 #include "SVGFactory.generated.h"
 
 UCLASS()
-class SVGIMPORTEREDITOR_API USVGFactory : public UFactory
+class USVGFactory : public UFactory
 {
 	GENERATED_BODY()
 
diff --git a/Engine/Plugins/Experimental/SVGImporter/Source/SVGImporterEditor/Private/SVGThumbnailRenderer.h b/Engine/Plugins/Experimental/SVGImporter/Source/SVGImporterEditor/Private/SVGThumbnailRenderer.h
--- a/Engine/Plugins/Experimental/SVGImporter/Source/SVGImporterEditor/Private/SVGThumbnailRenderer.h
+++ b/Engine/Plugins/Experimental/SVGImporter/Source/SVGImporterEditor/Private/SVGThumbnailRenderer.h
@@ -7,7 +7,7 @@
 #include "SVGThumbnailRenderer.generated.h"
 
 UCLASS()
-class SVGIMPORTEREDITOR_API USVGThumbnailRenderer : public UDefaultSizedThumbnailRenderer
+class USVGThumbnailRenderer : public UDefaultSizedThumbnailRenderer
 {
 	GENERATED_BODY()
 
