[SVG Importer] Better visualizer for SVG Actor extrude gizmo

- Details panel now keeps its focus on actor, and not on single shape component
- widget handle takes into account newly added scale property

Visualization of extrude handle is now done with a visualizer for a new editor-only component, which is not shown in details panel, so the panel focus stays on the owning actor.
Potentially, other widgets could be added to this visualizer.

diff --git a/Engine/Plugins/Experimental/SVGImporter/Source/SVGImporter/Private/SVGActor.cpp b/Engine/Plugins/Experimental/SVGImporter/Source/SVGImporter/Private/SVGActor.cpp
--- a/Engine/Plugins/Experimental/SVGImporter/Source/SVGImporter/Private/SVGActor.cpp
+++ b/Engine/Plugins/Experimental/SVGImporter/Source/SVGImporter/Private/SVGActor.cpp
@@ -20,6 +20,7 @@
 #include "LevelEditor.h"
 #include "Misc/TransactionObjectEvent.h"
 #include "Modules/ModuleManager.h"
+#include "SVGActorEditorComponent.h"
 #include "SVGImporterUtils.h"
 #include "ScopedTransaction.h"
 #include "Widgets/Notifications/SNotificationList.h"
@@ -94,6 +95,11 @@ ASVGActor::ASVGActor()
 	FillShapesRoot = CreateDefaultSubobject<USceneComponent>("FillsRoot");
 	FillShapesRoot->SetupAttachment(RootComponent);
 
+#if WITH_EDITOR
+	SVGEditorComponent = CreateDefaultSubobject<USVGActorEditorComponent>("SVGEditorComponent");
+	SVGEditorComponent->SetupAttachment(RootComponent);
+#endif
+
 	bMeshesShouldBeGenerated = true;
 	CurrExtrudeForDepth = 0.0f;
 
diff --git a/Engine/Plugins/Experimental/SVGImporter/Source/SVGImporter/Private/SVGActorEditorComponent.cpp b/Engine/Plugins/Experimental/SVGImporter/Source/SVGImporter/Private/SVGActorEditorComponent.cpp
new file mode 100644
--- /dev/null
+++ b/Engine/Plugins/Experimental/SVGImporter/Source/SVGImporter/Private/SVGActorEditorComponent.cpp
@@ -0,0 +1,16 @@
+// Copyright Epic Games, Inc. All Rights Reserved.
+
+#if WITH_EDITOR
+#include "SVGActorEditorComponent.h"
+#include "SVGActor.h"
+
+ASVGActor* USVGActorEditorComponent::GetSVGActor() const
+{
+	if (ASVGActor* SVGActor = GetTypedOuter<ASVGActor>())
+	{
+		return SVGActor;
+	}
+
+	return nullptr;
+}
+#endif
diff --git a/Engine/Plugins/Experimental/SVGImporter/Source/SVGImporter/Public/SVGActor.h b/Engine/Plugins/Experimental/SVGImporter/Source/SVGImporter/Public/SVGActor.h
--- a/Engine/Plugins/Experimental/SVGImporter/Source/SVGImporter/Public/SVGActor.h
+++ b/Engine/Plugins/Experimental/SVGImporter/Source/SVGImporter/Public/SVGActor.h
@@ -15,6 +15,7 @@
 #if WITH_EDITOR
 class FScopedTransaction;
 class FTransactionObjectEvent;
+class USVGActorEditorComponent;
 #endif
 
 class USVGData;
@@ -66,6 +67,10 @@ public:
 	virtual TArray<UDynamicMeshComponent*> GetSVGDynamicMeshes() override;
 	//~ End ASVGDynamicMeshesOwnerActor
 
+#if WITH_EDITOR
+	USVGActorEditorComponent* GetSVGEditorComponent() { return SVGEditorComponent; }
+#endif
+
 	TArray<TObjectPtr<USVGFillComponent>> GetFillComponents() const;
 	TArray<TObjectPtr<USVGStrokeComponent>> GetStrokeComponents() const;
 
@@ -273,6 +278,11 @@ protected:
 	UPROPERTY()
 	TObjectPtr<USceneComponent> StrokeShapesRoot;
 
+#if WITH_EDITORONLY_DATA
+	UPROPERTY()
+	TObjectPtr<USVGActorEditorComponent> SVGEditorComponent;
+#endif
+
 	UPROPERTY()
 	TArray<TObjectPtr<USVGDynamicMeshComponent>> ShapeComponents;
 
@@ -333,6 +343,7 @@ UPROPERTY()
 private:
 #if WITH_EDITOR
 	friend class FSVGDynamicMeshVisualizer;
+	friend class FSVGActorEditorComponentVisualizer;
 
 	void SetFillsExtrudeInteractive(float InFillsExtrude);
 	void SetStrokesExtrudeInteractive(float InStrokesExtrude);
diff --git a/Engine/Plugins/Experimental/SVGImporter/Source/SVGImporter/Public/SVGActorEditorComponent.h b/Engine/Plugins/Experimental/SVGImporter/Source/SVGImporter/Public/SVGActorEditorComponent.h
new file mode 100644
--- /dev/null
+++ b/Engine/Plugins/Experimental/SVGImporter/Source/SVGImporter/Public/SVGActorEditorComponent.h
@@ -0,0 +1,18 @@
+// Copyright Epic Games, Inc. All Rights Reserved.
+
+#pragma once
+
+#include "Components/SceneComponent.h"
+#include "SVGActorEditorComponent.generated.h"
+
+class ASVGActor;
+
+UCLASS(MinimalAPI)
+class USVGActorEditorComponent : public USceneComponent
+{
+	GENERATED_BODY()
+#if WITH_EDITOR
+public:
+	SVGIMPORTER_API ASVGActor* GetSVGActor() const;
+#endif
+};
diff --git a/Engine/Plugins/Experimental/SVGImporter/Source/SVGImporterEditor/Private/SVGActorEditorComponentVisualizer.cpp b/Engine/Plugins/Experimental/SVGImporter/Source/SVGImporterEditor/Private/SVGActorEditorComponentVisualizer.cpp
new file mode 100644
--- /dev/null
+++ b/Engine/Plugins/Experimental/SVGImporter/Source/SVGImporterEditor/Private/SVGActorEditorComponentVisualizer.cpp
@@ -0,0 +1,245 @@
+// Copyright Epic Games, Inc. All Rights Reserved.
+
+#include "SVGActorEditorComponentVisualizer.h"
+#include "EditorViewportClient.h"
+#include "HitProxies.h"
+#include "SVGActor.h"
+#include "SVGActorEditorComponent.h"
+#include "SceneManagement.h"
+
+float FSVGActorEditorComponentVisualizer::FillsExtrudeMin = 0.01f;
+float FSVGActorEditorComponentVisualizer::FillsExtrudeMax = 20.0f;
+float FSVGActorEditorComponentVisualizer::StrokesExtrudeMin = 0.01f;
+float FSVGActorEditorComponentVisualizer::StrokesExtrudeMax = 20.0f;
+
+struct HSVGActorExtrudeHitProxy : HComponentVisProxy
+{
+	DECLARE_HIT_PROXY();
+
+	HSVGActorExtrudeHitProxy(const UActorComponent* InComponent, ASVGActor* InSVGActor)
+		: HComponentVisProxy(InComponent, HPP_Wireframe)
+		, SVGActorWeak(InSVGActor)
+	{ }
+
+	TWeakObjectPtr<ASVGActor> SVGActorWeak;
+};
+
+IMPLEMENT_HIT_PROXY(HSVGActorExtrudeHitProxy, HComponentVisProxy)
+
+FSVGActorEditorComponentVisualizer::FSVGActorEditorComponentVisualizer()
+{
+	bIsExtruding = false;
+	UpdateMinMaxExtrudeValues();
+}
+
+void FSVGActorEditorComponentVisualizer::DrawVisualization(const UActorComponent* InComponent, const FSceneView* InView, FPrimitiveDrawInterface* InPDI)
+{
+	constexpr static float ExtrudeHandleSize = 10.0f;
+
+	const USVGActorEditorComponent* SVGActorEditorComponent = Cast<USVGActorEditorComponent>(InComponent);
+	if (!SVGActorEditorComponent)
+	{
+		InPDI->SetHitProxy(nullptr);
+		EndEditing();
+		return;
+	}
+
+	if (ASVGActor* SVGActor = Cast<ASVGActor>(SVGActorEditorComponent->GetSVGActor()))
+	{
+		if (SVGActor->ExtrudeType == ESVGExtrudeType::None || SVGActor->RenderMode == ESVGRenderMode::Texture2D)
+		{
+			EndEditing();
+		}
+		else
+		{
+			SVGActorWeak = SVGActor;
+			SVGEditorComponentWeak = SVGActor->GetSVGEditorComponent();
+
+			if (SVGEditorComponentWeak.IsValid())
+			{
+				const FVector ProxyHandleLocation = GetExtrudeWidgetLocation();
+				const FVector LineStartLocation = GetExtrudeSurfaceLocation();
+				InPDI->DrawLine(LineStartLocation, ProxyHandleLocation, FLinearColor::White, SDPG_Foreground);
+
+				InPDI->SetHitProxy(new HSVGActorExtrudeHitProxy(SVGActorEditorComponent, SVGActor));
+				InPDI->DrawPoint(ProxyHandleLocation, FLinearColor::White, ExtrudeHandleSize, SDPG_Foreground);
+				InPDI->SetHitProxy(nullptr);
+			}
+		}
+	}
+}
+
+bool FSVGActorEditorComponentVisualizer::VisProxyHandleClick(FEditorViewportClient* InViewportClient, HComponentVisProxy* InVisProxy, const FViewportClick& InClick)
+{
+	EndEditing();
+
+	if (InClick.GetKey() == EKeys::LeftMouseButton && InVisProxy)
+	{
+		if (InVisProxy->IsA(HSVGActorExtrudeHitProxy::StaticGetType()))
+		{
+			if (const HSVGActorExtrudeHitProxy* const ExtrudeProxy = static_cast<HSVGActorExtrudeHitProxy*>(InVisProxy))
+			{
+				if (ASVGActor* SVGActor = ExtrudeProxy->SVGActorWeak.Get())
+				{
+					if (GEditor)
+					{
+						GEditor->SelectActor(SVGActor, true, false);
+					}
+
+					bIsExtruding = true;
+					SVGActorWeak = SVGActor;
+					SVGEditorComponentWeak = Cast<USVGActorEditorComponent>(const_cast<UActorComponent*>(ExtrudeProxy->Component.Get()));
+					InViewportClient->SetWidgetMode(UE::Widget::WM_None);
+
+					return true;
+				}
+			}
+		}
+	}
+
+	return false;
+}
+
+bool FSVGActorEditorComponentVisualizer::GetWidgetLocation(const FEditorViewportClient* InViewportClient, FVector& OutLocation) const
+{
+	if (bIsExtruding)
+	{
+		if (SVGActorWeak.IsValid())
+		{
+			OutLocation = GetExtrudeWidgetLocation();
+			return true;
+		}
+	}
+
+	return false;
+}
+
+bool FSVGActorEditorComponentVisualizer::HandleInputDelta(FEditorViewportClient* InViewportClient, FViewport* InViewport, FVector& InDeltaTranslate, FRotator& InDeltaRotate, FVector& InDeltaScale)
+{
+	if (!GetEditedComponent() || !GetEditedComponent()->GetOwner())
+	{
+		EndEditing();
+		return false;
+	}
+
+	if (!InDeltaTranslate.IsZero())
+	{
+		if (bIsExtruding)
+		{
+			if (ASVGActor* SVGActor = SVGActorWeak.Get())
+			{
+				SVGActor->Modify();
+
+				float FillsExtrudeDepth = SVGActor->GetFillsExtrude() - (InDeltaTranslate.X / SVGActor->GetScale());
+				FillsExtrudeDepth = FMath::Clamp(FillsExtrudeDepth, FillsExtrudeMin, FillsExtrudeMax);
+				SVGActor->SetFillsExtrudeInteractive(FillsExtrudeDepth);
+
+				float StrokesExtrudeDepth = SVGActor->GetStrokesExtrude() - (InDeltaTranslate.X / SVGActor->GetScale());
+				StrokesExtrudeDepth = FMath::Clamp(StrokesExtrudeDepth, StrokesExtrudeMin, StrokesExtrudeMax);
+				SVGActor->SetStrokesExtrudeInteractive(StrokesExtrudeDepth);
+
+				return true;
+			}
+		}
+	}
+
+	EndEditing();
+	return false;
+}
+
+void FSVGActorEditorComponentVisualizer::EndEditing()
+{
+	bIsExtruding = false;
+}
+
+void FSVGActorEditorComponentVisualizer::TrackingStopped(FEditorViewportClient* InViewportClient, bool bInDidMove)
+{
+	if (ASVGActor* SVGActor = SVGActorWeak.Get())
+	{
+		SVGActor->Modify();
+
+		const float ExtrudeDepth = FMath::Clamp(SVGActor->GetFillsExtrude(), FillsExtrudeMin, FillsExtrudeMax);
+		SVGActor->SetFillsExtrude(ExtrudeDepth);
+
+		const float StrokesExtrudeDepth = FMath::Clamp(SVGActor->GetStrokesExtrude(), StrokesExtrudeMin, StrokesExtrudeMax);
+		SVGActor->SetStrokesExtrude(StrokesExtrudeDepth);
+	}
+
+	FComponentVisualizer::TrackingStopped(InViewportClient, bInDidMove);
+}
+
+UActorComponent* FSVGActorEditorComponentVisualizer::GetEditedComponent() const
+{
+	return SVGEditorComponentWeak.Get();
+}
+
+FVector FSVGActorEditorComponentVisualizer::GetExtrudeSurfaceLocation() const
+{
+	if (const ASVGActor* SVGActor = SVGActorWeak.Get())
+	{
+		float ExtrudeDepth = FMath::Max(SVGActor->GetFillsExtrude(), SVGActor->GetStrokesExtrude());
+
+		if (SVGActor->ExtrudeType == ESVGExtrudeType::FrontBackMirror)
+		{
+			ExtrudeDepth *= 0.5f;
+		}
+
+		const FVector SurfaceOffset = -SVGActor->GetActorForwardVector() * ExtrudeDepth * SVGActor->GetScale();
+
+		return SVGActor->GetActorLocation() + SurfaceOffset;
+	}
+
+	return {};
+}
+
+FVector FSVGActorEditorComponentVisualizer::GetExtrudeWidgetLocation() const
+{
+	constexpr static float ExtrudeHandleOffset = 10.0f;
+	const FVector SurfaceLocation = GetExtrudeSurfaceLocation();
+
+	if (const ASVGActor* SVGActor = SVGActorWeak.Get())
+	{
+		const float ExtrudeDepth = FMath::Max(SVGActor->GetFillsExtrude(), SVGActor->GetStrokesExtrude());
+		FVector HandleOffset = SVGActor->GetActorForwardVector() * ExtrudeHandleOffset;
+
+		if (ExtrudeDepth >= 0)
+		{
+			HandleOffset = -HandleOffset;
+		}
+
+		return SurfaceLocation + HandleOffset * SVGActor->GetScale();
+	}
+
+	return {};
+}
+
+void FSVGActorEditorComponentVisualizer::UpdateMinMaxExtrudeValues()
+{
+	// Get Min and Max Extrude values from Properties Metadata, to match what the UI allows
+
+	if (const FProperty* FillsExtrudeProperty = ASVGActor::StaticClass()->FindPropertyByName(TEXT("FillsExtrude")))
+	{
+		if (FillsExtrudeProperty->HasMetaData("UIMin"))
+		{
+			FillsExtrudeMin = FCString::Atof(*FillsExtrudeProperty->GetMetaData("UIMin"));
+		}
+
+		if (FillsExtrudeProperty->HasMetaData("UIMax"))
+		{
+			FillsExtrudeMax = FCString::Atof(*FillsExtrudeProperty->GetMetaData("UIMax"));
+		}
+	}
+
+	if (const FProperty* StrokesExtrudeProperty = ASVGActor::StaticClass()->FindPropertyByName(TEXT("StrokesExtrude")))
+	{
+		if (StrokesExtrudeProperty->HasMetaData("UIMin"))
+		{
+			StrokesExtrudeMin = FCString::Atof(*StrokesExtrudeProperty->GetMetaData("UIMin"));
+		}
+
+		if (StrokesExtrudeProperty->HasMetaData("UIMax"))
+		{
+			StrokesExtrudeMax = FCString::Atof(*StrokesExtrudeProperty->GetMetaData("UIMax"));
+		}
+	}
+}
diff --git a/Engine/Plugins/Experimental/SVGImporter/Source/SVGImporterEditor/Private/SVGActorEditorComponentVisualizer.h b/Engine/Plugins/Experimental/SVGImporter/Source/SVGImporterEditor/Private/SVGActorEditorComponentVisualizer.h
new file mode 100644
--- /dev/null
+++ b/Engine/Plugins/Experimental/SVGImporter/Source/SVGImporterEditor/Private/SVGActorEditorComponentVisualizer.h
@@ -0,0 +1,39 @@
+// Copyright Epic Games, Inc. All Rights Reserved.
+
+#pragma once
+
+#include "ComponentVisualizer.h"
+
+class ASVGActor;
+class USVGActorEditorComponent;
+
+class FSVGActorEditorComponentVisualizer : public FComponentVisualizer
+{
+public:
+	static void UpdateMinMaxExtrudeValues();
+
+	FSVGActorEditorComponentVisualizer();
+
+	//~ Begin FComponentVisualizer
+	virtual void DrawVisualization(const UActorComponent* InComponent, const FSceneView* InView, FPrimitiveDrawInterface* InPDI) override;
+	virtual bool VisProxyHandleClick(FEditorViewportClient* InViewportClient, HComponentVisProxy* InVisProxy, const FViewportClick& InClick) override;
+	virtual bool GetWidgetLocation(const FEditorViewportClient* InViewportClient, FVector& OutLocation) const override;
+	virtual bool HandleInputDelta(FEditorViewportClient* InViewportClient, FViewport* InViewport, FVector& InDeltaTranslate, FRotator& InDeltaRotate, FVector& InDeltaScale) override;
+	virtual void EndEditing() override;
+	virtual void TrackingStopped(FEditorViewportClient* InViewportClient, bool bInDidMove) override;
+	virtual UActorComponent* GetEditedComponent() const override;
+	//~ End FComponentVisualizer
+
+private:
+	FVector GetExtrudeWidgetLocation() const;
+	FVector GetExtrudeSurfaceLocation() const;
+
+	static float FillsExtrudeMin;
+	static float FillsExtrudeMax;
+	static float StrokesExtrudeMin;
+	static float StrokesExtrudeMax;
+
+	TWeakObjectPtr<USVGActorEditorComponent> SVGEditorComponentWeak;
+	TWeakObjectPtr<ASVGActor> SVGActorWeak;
+	bool bIsExtruding = false;
+};
diff --git a/Engine/Plugins/Experimental/SVGImporter/Source/SVGImporterEditor/Private/SVGDynamicMeshVisualizer.cpp b/Engine/Plugins/Experimental/SVGImporter/Source/SVGImporterEditor/Private/SVGDynamicMeshVisualizer.cpp
--- a/Engine/Plugins/Experimental/SVGImporter/Source/SVGImporterEditor/Private/SVGDynamicMeshVisualizer.cpp
+++ b/Engine/Plugins/Experimental/SVGImporter/Source/SVGImporterEditor/Private/SVGDynamicMeshVisualizer.cpp
@@ -1,301 +1,34 @@
 // Copyright Epic Games, Inc. All Rights Reserved.
 
 #include "SVGDynamicMeshVisualizer.h"
-#include "EditorViewportClient.h"
-#include "HitProxies.h"
 #include "ProceduralMeshes/SVGDynamicMeshComponent.h"
-#include "SVGActor.h"
 #include "SceneManagement.h"
 
-float FSVGDynamicMeshVisualizer::FillsExtrudeMin;
-float FSVGDynamicMeshVisualizer::FillsExtrudeMax;
-float FSVGDynamicMeshVisualizer::StrokesExtrudeMin;
-float FSVGDynamicMeshVisualizer::StrokesExtrudeMax;
-
-struct HSVGActorExtrudeHitProxy : HComponentVisProxy
-{
-	DECLARE_HIT_PROXY();
-
-	HSVGActorExtrudeHitProxy(const UActorComponent* InComponent, ASVGActor* InSVGActor, bool bInIsFirstComponent)
-		: HComponentVisProxy(InComponent, HPP_Wireframe)
-		, SVGActorWeak(InSVGActor)
-		, bIsFirstComponent(bInIsFirstComponent)
-	{ }
-
-	TWeakObjectPtr<ASVGActor> SVGActorWeak;
-	bool bIsFirstComponent = false;
-};
-
-IMPLEMENT_HIT_PROXY(HSVGActorExtrudeHitProxy, HComponentVisProxy)
-
-FSVGDynamicMeshVisualizer::FSVGDynamicMeshVisualizer()
-{
-	bIsFirstSVGComponent = false;
-	bIsExtruding = false;
-
-	UpdateMinMaxExtrudeValues();
-}
-
 void FSVGDynamicMeshVisualizer::DrawVisualization(const UActorComponent* InComponent, const FSceneView* InView, FPrimitiveDrawInterface* InPDI)
 {
-	constexpr static float ExtrudeHandleSize = 10.0f;
-
-	const USVGDynamicMeshComponent* SVGDynamicMeshComponent = Cast<USVGDynamicMeshComponent>(InComponent);
-	if (!SVGDynamicMeshComponent)
-	{
-		InPDI->SetHitProxy(nullptr);
-		EndEditing();
-		return;
-	}
-
-	if (ASVGActor* SVGActor = Cast<ASVGActor>(SVGDynamicMeshComponent->GetOwner()))
-	{
-		if (SVGActor->ExtrudeType == ESVGExtrudeType::None)
-		{
-			EndEditing();
-		}
-		else
-		{
-			SVGActorWeak = SVGActor;
-			const TArray<UDynamicMeshComponent*>& SVGShapeComponents = SVGActor->GetSVGDynamicMeshes();
-
-			if (!SVGShapeComponents.IsEmpty())
-			{
-				if (SVGDynamicMeshComponent == SVGShapeComponents[0])
-				{
-					const FVector ProxyHandleLocation = GetExtrudeWidgetLocation();
-					const FVector LineStartLocation = GetExtrudeSurfaceLocation();
-					InPDI->DrawLine(LineStartLocation, ProxyHandleLocation, FLinearColor::White, SDPG_Foreground);
-
-					InPDI->SetHitProxy(new HSVGActorExtrudeHitProxy(SVGDynamicMeshComponent, SVGActor, true));
-					InPDI->DrawPoint(ProxyHandleLocation, FLinearColor::White, ExtrudeHandleSize, SDPG_Foreground);
-					InPDI->SetHitProxy(nullptr);
-				}
-			}
-		}
-	}
-
-	// Draw box to highlight bevel setting while in interactive mode
-	if (SVGDynamicMeshComponent->bIsBevelBeingEdited)
-	{
-		FVector BoxLocation = SVGDynamicMeshComponent->GetComponentLocation() - FVector(SVGDynamicMeshComponent->Bounds.BoxExtent.X, 0, 0) + FVector(SVGDynamicMeshComponent->Bevel/2.0, 0, 0);
-
-		if (SVGDynamicMeshComponent->ExtrudeType == ESVGExtrudeType::FrontFaceOnly)
-		{
-			BoxLocation -= FVector(SVGDynamicMeshComponent->Bounds.BoxExtent.X, 0, 0);
-		}
-
-		FBox Box;
-		Box.Max = BoxLocation + FVector(SVGDynamicMeshComponent->Bevel/2.0f, SVGDynamicMeshComponent->Bounds.BoxExtent.Y, SVGDynamicMeshComponent->Bounds.BoxExtent.Z);
-		Box.Min = BoxLocation - FVector(SVGDynamicMeshComponent->Bevel/2.0f, SVGDynamicMeshComponent->Bounds.BoxExtent.Y, SVGDynamicMeshComponent->Bounds.BoxExtent.Z);
-
-		DrawWireBox(InPDI, Box, FColor::Green, SDPG_Foreground);
-	}
-}
-
-bool FSVGDynamicMeshVisualizer::VisProxyHandleClick(FEditorViewportClient* InViewportClient, HComponentVisProxy* VisProxy, const FViewportClick& Click)
-{
-	EndEditing();
-
-	if (Click.GetKey() == EKeys::LeftMouseButton && VisProxy)
+	if (const USVGDynamicMeshComponent* SVGMeshComp = Cast<USVGDynamicMeshComponent>(InComponent))
 	{
-		if (VisProxy->IsA(HSVGActorExtrudeHitProxy::StaticGetType()))
+		// draw box to highlight bevel setting while in interactive mode
+		if (SVGMeshComp->bIsBevelBeingEdited)
 		{
-			if (const HSVGActorExtrudeHitProxy* const ExtrudeProxy = static_cast<HSVGActorExtrudeHitProxy*>(VisProxy))
-			{
-				if (ASVGActor* SVGActor = ExtrudeProxy->SVGActorWeak.Get())
-				{
-					if (ExtrudeProxy->bIsFirstComponent)
-					{
-						GEditor->SelectActor(SVGActor, true, false);
-						bIsExtruding = true;
-						bIsFirstSVGComponent = ExtrudeProxy->bIsFirstComponent;
-						SVGActorWeak = SVGActor;
-						SVGMeshComponent = Cast<USVGDynamicMeshComponent>(const_cast<UActorComponent*>(ExtrudeProxy->Component.Get()));
-						InViewportClient->SetWidgetMode(UE::Widget::WM_None);
-						return true;
-					}
-				}
-			}
-		}
-	}
-
-	return false;
-}
+			FVector BoxLocation = SVGMeshComp->GetComponentLocation() - FVector(SVGMeshComp->Bounds.BoxExtent.X, 0, 0) + FVector(SVGMeshComp->Bevel * 0.5f, 0, 0);
 
-bool FSVGDynamicMeshVisualizer::GetWidgetLocation(const FEditorViewportClient* ViewportClient, FVector& OutLocation) const
-{
-	if (bIsFirstSVGComponent && bIsExtruding)
-	{
-		if (SVGActorWeak.IsValid())
-		{
-			OutLocation = GetExtrudeWidgetLocation();
-			return true;
-		}
-	}
-
-	return false;
-}
-
-bool FSVGDynamicMeshVisualizer::HandleInputDelta(FEditorViewportClient* InViewportClient, FViewport* InViewport, FVector& InDeltaTranslate, FRotator& InDeltaRotate, FVector& InDeltaScale)
-{
-	if (!GetEditedComponent() || !GetEditedComponent()->GetOwner())
-	{
-		EndEditing();
-		return false;
-	}
-
-	if (!InDeltaTranslate.IsZero())
-	{
-		if (bIsFirstSVGComponent)
-		{
-			if (bIsExtruding)
+			if (SVGMeshComp->ExtrudeType == ESVGExtrudeType::FrontFaceOnly)
 			{
-				if (ASVGActor* SVGActor = SVGActorWeak.Get())
-				{
-					SVGActor->Modify();
-
-					float FillsExtrudeDepth = SVGActor->GetFillsExtrude() - InDeltaTranslate.X;
-					FillsExtrudeDepth = FMath::Clamp(FillsExtrudeDepth, FillsExtrudeMin, FillsExtrudeMax);
-					SVGActor->SetFillsExtrudeInteractive(FillsExtrudeDepth);
-
-					float StrokesExtrudeDepth = SVGActor->GetStrokesExtrude() - InDeltaTranslate.X;
-					StrokesExtrudeDepth = FMath::Clamp(StrokesExtrudeDepth, StrokesExtrudeMin, StrokesExtrudeMax);
-					SVGActor->SetStrokesExtrudeInteractive(StrokesExtrudeDepth);
-
-					return true;
-				}
+				BoxLocation -= FVector(SVGMeshComp->Bounds.BoxExtent.X, 0, 0);
 			}
-		}
-	}
-
-	EndEditing();
-	return false;
-}
-
-void FSVGDynamicMeshVisualizer::EndEditing()
-{
-	bIsFirstSVGComponent = false;
-	bIsExtruding = false;
-}
-
-void FSVGDynamicMeshVisualizer::TrackingStopped(FEditorViewportClient* InViewportClient, bool bInDidMove)
-{
-	if (ASVGActor* SVGActor = SVGActorWeak.Get())
-	{
-		SVGActor->Modify();
-
-		const float ExtrudeDepth = FMath::Clamp(SVGActor->GetFillsExtrude(), FillsExtrudeMin, FillsExtrudeMax);
-		SVGActor->SetFillsExtrude(ExtrudeDepth);
-
-		const float StrokesExtrudeDepth = FMath::Clamp(SVGActor->GetStrokesExtrude(), StrokesExtrudeMin, StrokesExtrudeMax);
-		SVGActor->SetStrokesExtrude(StrokesExtrudeDepth);
-	}
-
-	FComponentVisualizer::TrackingStopped(InViewportClient, bInDidMove);
-}
-
-UActorComponent* FSVGDynamicMeshVisualizer::GetEditedComponent() const
-{
-	return SVGMeshComponent.Get();
-}
-
-FVector FSVGDynamicMeshVisualizer::GetExtrudeSurfaceLocation() const
-{
-	if (const ASVGActor* SVGActor = SVGActorWeak.Get())
-	{
-		float ExtrudeDepth = FMath::Max(SVGActor->GetFillsExtrude(), SVGActor->GetStrokesExtrude());
-
-		if (SVGActor->ExtrudeType == ESVGExtrudeType::FrontBackMirror)
-		{
-			ExtrudeDepth *= 0.5f;
-		}
-
-		const FVector SurfaceOffset = -SVGActor->GetActorForwardVector() * ExtrudeDepth;
-
-		return SVGActor->GetActorLocation() + SurfaceOffset;
-	}
-
-	return {};
-}
-
-FVector FSVGDynamicMeshVisualizer::GetExtrudeWidgetLocation() const
-{
-	constexpr static float ExtrudeHandleOffset = 10.0f;
-	const FVector SurfaceLocation = GetExtrudeSurfaceLocation();
-
-	if (const ASVGActor* SVGActor = SVGActorWeak.Get())
-	{
-		const float ExtrudeDepth = FMath::Max(SVGActor->GetFillsExtrude(), SVGActor->GetStrokesExtrude());
-		FVector HandleOffset = SVGActor->GetActorForwardVector() * ExtrudeHandleOffset;
-
-		if (ExtrudeDepth >= 0)
-		{
-			HandleOffset = -HandleOffset;
-		}
-
-		return SurfaceLocation + HandleOffset;
-	}
-
-	return {};
-}
-
-void FSVGDynamicMeshVisualizer::UpdateMinMaxExtrudeValues()
-{
-	static float DefaultExtrudeMin = 0.01f;
-	static float DefaultExtrudeMax = 20.0f;
 
-	// Get Min and Max Extrude values from Properties Metadata, to match what the UI allows
+			FBox Box;
+			Box.Max = BoxLocation + FVector(SVGMeshComp->Bevel * 0.5f, SVGMeshComp->Bounds.BoxExtent.Y, SVGMeshComp->Bounds.BoxExtent.Z);
+			Box.Min = BoxLocation - FVector(SVGMeshComp->Bevel * 0.5f, SVGMeshComp->Bounds.BoxExtent.Y, SVGMeshComp->Bounds.BoxExtent.Z);
 
-	if (const FProperty* FillsExtrudeProperty = ASVGActor::StaticClass()->FindPropertyByName(TEXT("FillsExtrude")))
-	{
-		if (FillsExtrudeProperty->HasMetaData("UIMin"))
-		{
-			FillsExtrudeMin = FCString::Atof(*FillsExtrudeProperty->GetMetaData("UIMin"));
-		}
-		else
-		{
-			FillsExtrudeMin = DefaultExtrudeMin;
+			DrawWireBox(InPDI, Box, FColor::Green, SDPG_Foreground);
 		}
-
-		if (FillsExtrudeProperty->HasMetaData("UIMax"))
-		{
-			FillsExtrudeMax = FCString::Atof(*FillsExtrudeProperty->GetMetaData("UIMax"));
-		}
-		else
+		// in other cases, if the SVG mesh is selected, draw a box around it
+		else if (SVGMeshComp->IsSelectedInEditor())
 		{
-			FillsExtrudeMax = DefaultExtrudeMax;
+			const FBox Bounds(SVGMeshComp->Bounds.GetBox());
+			DrawWireBox(InPDI, Bounds, FColor::Red, SDPG_Foreground);
 		}
 	}
-	else
-	{
-		FillsExtrudeMin = DefaultExtrudeMin;
-		FillsExtrudeMax = DefaultExtrudeMax;
-	}
-
-	if (const FProperty* StrokesExtrudeProperty = ASVGActor::StaticClass()->FindPropertyByName(TEXT("StrokesExtrude")))
-	{
-		if (StrokesExtrudeProperty->HasMetaData("UIMin"))
-		{
-			StrokesExtrudeMin = FCString::Atof(*StrokesExtrudeProperty->GetMetaData("UIMin"));
-		}
-		else
-		{
-			StrokesExtrudeMin = DefaultExtrudeMin;
-		}
-
-		if (StrokesExtrudeProperty->HasMetaData("UIMax"))
-		{
-			StrokesExtrudeMax = FCString::Atof(*StrokesExtrudeProperty->GetMetaData("UIMax"));
-		}
-		else
-		{
-			StrokesExtrudeMax = DefaultExtrudeMax;
-		}
-	}
-	else
-	{
-		StrokesExtrudeMin = DefaultExtrudeMin;
-		StrokesExtrudeMax = DefaultExtrudeMax;
-	}
 }
diff --git a/Engine/Plugins/Experimental/SVGImporter/Source/SVGImporterEditor/Private/SVGDynamicMeshVisualizer.h b/Engine/Plugins/Experimental/SVGImporter/Source/SVGImporterEditor/Private/SVGDynamicMeshVisualizer.h
--- a/Engine/Plugins/Experimental/SVGImporter/Source/SVGImporterEditor/Private/SVGDynamicMeshVisualizer.h
+++ b/Engine/Plugins/Experimental/SVGImporter/Source/SVGImporterEditor/Private/SVGDynamicMeshVisualizer.h
@@ -4,37 +4,8 @@
 
 #include "ComponentVisualizer.h"
 
-class ASVGActor;
-class USVGDynamicMeshComponent;
-
 class FSVGDynamicMeshVisualizer : public FComponentVisualizer
 {
 public:
-	static void UpdateMinMaxExtrudeValues();
-
-	FSVGDynamicMeshVisualizer();
-
-	//~ Begin FComponentVisualizer
 	virtual void DrawVisualization(const UActorComponent* InComponent, const FSceneView* InView, FPrimitiveDrawInterface* InPDI) override;
-	virtual bool VisProxyHandleClick(FEditorViewportClient* InViewportClient, HComponentVisProxy* VisProxy, const FViewportClick& Click) override;
-	virtual bool GetWidgetLocation(const FEditorViewportClient* ViewportClient, FVector& OutLocation) const override;
-	virtual bool HandleInputDelta(FEditorViewportClient* InViewportClient, FViewport* InViewport, FVector& InDeltaTranslate, FRotator& InDeltaRotate, FVector& InDeltaScale) override;
-	virtual void EndEditing() override;
-	virtual void TrackingStopped(FEditorViewportClient* InViewportClient, bool bInDidMove) override;
-	virtual UActorComponent* GetEditedComponent() const override;
-	//~ End FComponentVisualizer
-
-private:
-	FVector GetExtrudeWidgetLocation() const;
-	FVector GetExtrudeSurfaceLocation() const;
-
-	static float FillsExtrudeMin;
-	static float FillsExtrudeMax;
-	static float StrokesExtrudeMin;
-	static float StrokesExtrudeMax;
-
-	TWeakObjectPtr<USVGDynamicMeshComponent> SVGMeshComponent;
-	TWeakObjectPtr<ASVGActor> SVGActorWeak;
-	bool bIsFirstSVGComponent = false;
-	bool bIsExtruding = false;
 };
diff --git a/Engine/Plugins/Experimental/SVGImporter/Source/SVGImporterEditor/Private/SVGImporterEditorModule.cpp b/Engine/Plugins/Experimental/SVGImporter/Source/SVGImporterEditor/Private/SVGImporterEditorModule.cpp
--- a/Engine/Plugins/Experimental/SVGImporter/Source/SVGImporterEditor/Private/SVGImporterEditorModule.cpp
+++ b/Engine/Plugins/Experimental/SVGImporter/Source/SVGImporterEditor/Private/SVGImporterEditorModule.cpp
@@ -19,6 +19,8 @@
 #include "PropertyEditorModule.h"
 #include "SVGActor.h"
 #include "SVGActorContextMenu.h"
+#include "SVGActorEditorComponent.h"
+#include "SVGActorEditorComponentVisualizer.h"
 #include "SVGData.h"
 #include "SVGDynamicMeshVisualizer.h"
 #include "SVGImporter.h"
@@ -167,12 +169,11 @@ void FSVGImporterEditorModule::PostEngineInit()
 {
 	if (GUnrealEd)
 	{
-		// Make a new instance of the visualizer
-		const TSharedPtr<FComponentVisualizer> SVGDynMeshVisualizer = MakeShared<FSVGDynamicMeshVisualizer>();
-
-		// Register it to our specific component class
+		const TSharedRef<FComponentVisualizer> SVGDynMeshVisualizer = MakeShared<FSVGDynamicMeshVisualizer>();
 		GUnrealEd->RegisterComponentVisualizer(USVGDynamicMeshComponent::StaticClass()->GetFName(), SVGDynMeshVisualizer);
-		SVGDynMeshVisualizer->OnRegister();
+
+		const TSharedRef<FComponentVisualizer> SVGActorEditorComponentVisualizer = MakeShared<FSVGActorEditorComponentVisualizer>();
+		GUnrealEd->RegisterComponentVisualizer(USVGActorEditorComponent::StaticClass()->GetFName(), SVGActorEditorComponentVisualizer);
 	}
 
 	// Register SVG category for needed types
