[SVG Importer] Regression: Fix issue with split SVG shapes being incorrectly repositioned on PostLoad

diff --git a/Engine/Plugins/Experimental/SVGImporter/Source/SVGImporter/Private/ProceduralMeshes/SVGDynamicMeshComponent.cpp b/Engine/Plugins/Experimental/SVGImporter/Source/SVGImporter/Private/ProceduralMeshes/SVGDynamicMeshComponent.cpp
--- a/Engine/Plugins/Experimental/SVGImporter/Source/SVGImporter/Private/ProceduralMeshes/SVGDynamicMeshComponent.cpp
+++ b/Engine/Plugins/Experimental/SVGImporter/Source/SVGImporter/Private/ProceduralMeshes/SVGDynamicMeshComponent.cpp
@@ -10,6 +10,7 @@
 #include "Materials/Material.h"
 #include "Materials/MaterialInstanceDynamic.h"
 #include "Misc/Paths.h"
+#include "SVGActor.h"
 #include "UObject/ConstructorHelpers.h"
 
 #if WITH_EDITOR
@@ -109,7 +110,10 @@ void USVGDynamicMeshComponent::ApplyScale()
 		MeshTransforms::Translate(EditMesh, -Center);
 	});
 
-	SetRelativeLocation(FVector(0.0f, InternalCenter.Y, InternalCenter.Z) * Scale);
+	if (Cast<ASVGActor>(GetOwner()))
+	{
+		SetRelativeLocation(FVector(0.0f, InternalCenter.Y, InternalCenter.Z) * Scale);
+	}
 }
 
 void USVGDynamicMeshComponent::ScaleShape(float InScale)
