[SVG Importer] Imported SVGs textures should be centered

Texture content will now be centered with respect to svg canvas size.

Note: this might still end up having a different result when compared to 3D Geometry generated from the same SVG, for those cases in which SVG shapes are not centered in the original SVG source.
This happens because 3D SVG geometry is centered with respect to its bounding box, so it's not currently taking into account the empty space of the svg viewbox.
E.g. texture generated from SVG defining a 400x400 viewbox containing a 100x100 rectangle placed in 0;0. Something like that will still be centered wrt viewbox, and not the single rectangle it contains. The geometry will be centered wrt the rectangle bounds.

diff --git a/Engine/Plugins/Experimental/SVGImporter/Source/SVGImporter/Private/SVGImporterUtils.cpp b/Engine/Plugins/Experimental/SVGImporter/Source/SVGImporter/Private/SVGImporterUtils.cpp
--- a/Engine/Plugins/Experimental/SVGImporter/Source/SVGImporter/Private/SVGImporterUtils.cpp
+++ b/Engine/Plugins/Experimental/SVGImporter/Source/SVGImporter/Private/SVGImporterUtils.cpp
@@ -573,6 +573,9 @@ UTexture2D* FSVGImporterUtils::CreateSVGTexture(const FString& InSVGString, UObj
 	const float TextureDim = FMath::RoundUpToPowerOfTwo(MaxDim);
 
 	const float Scale = TextureDim/MaxDim;
+	const float XOffset = Scale * FMath::Max(0, (MaxDim - Image->width)/2.0f);
+	const float YOffset = Scale * FMath::Max(0, (MaxDim - Image->height)/2.0f);
+
 	TArray<uint8> PixelData;
 	PixelData.AddUninitialized(TextureDim * TextureDim * SVG_BPP);
 	const int32 Stride = TextureDim * SVG_BPP;
@@ -583,7 +586,7 @@ UTexture2D* FSVGImporterUtils::CreateSVGTexture(const FString& InSVGString, UObj
 		return nullptr;
 	}
 
-	nsvgRasterizeFull(Rasterizer, Image, 0, 0, Scale, Scale, PixelData.GetData(), TextureDim, TextureDim, Stride);
+	nsvgRasterizeFull(Rasterizer, Image, XOffset, YOffset, Scale, Scale, PixelData.GetData(), TextureDim, TextureDim, Stride);
 
 	nsvgDeleteRasterizer(Rasterizer);
 	nsvgDelete(Image);
