[SVG Importer] SVG Actor component visualizer marking level dirty

diff --git a/Engine/Plugins/Experimental/SVGImporter/Source/SVGImporterEditor/Private/SVGActorEditorComponentVisualizer.cpp b/Engine/Plugins/Experimental/SVGImporter/Source/SVGImporterEditor/Private/SVGActorEditorComponentVisualizer.cpp
--- a/Engine/Plugins/Experimental/SVGImporter/Source/SVGImporterEditor/Private/SVGActorEditorComponentVisualizer.cpp
+++ b/Engine/Plugins/Experimental/SVGImporter/Source/SVGImporterEditor/Private/SVGActorEditorComponentVisualizer.cpp
@@ -154,6 +154,11 @@ void FSVGActorEditorComponentVisualizer::EndEditing()
 
 void FSVGActorEditorComponentVisualizer::TrackingStopped(FEditorViewportClient* InViewportClient, bool bInDidMove)
 {
+	if (!bIsExtruding)
+	{
+		return;
+	}
+
 	if (ASVGActor* SVGActor = SVGActorWeak.Get())
 	{
 		SVGActor->Modify();
