Fixed uninitialized FLinearColor.

diff --git a/Engine/Plugins/Experimental/SVGImporter/Source/SVGImporter/Private/ProceduralMeshes/JoinedSVGDynamicMeshComponent.cpp b/Engine/Plugins/Experimental/SVGImporter/Source/SVGImporter/Private/ProceduralMeshes/JoinedSVGDynamicMeshComponent.cpp
--- a/Engine/Plugins/Experimental/SVGImporter/Source/SVGImporter/Private/ProceduralMeshes/JoinedSVGDynamicMeshComponent.cpp
+++ b/Engine/Plugins/Experimental/SVGImporter/Source/SVGImporter/Private/ProceduralMeshes/JoinedSVGDynamicMeshComponent.cpp
@@ -276,7 +276,7 @@ void UJoinedSVGDynamicMeshComponent::UpdateMaterials(bool bInRefreshInstances)
 		{
 			UMaterialInstanceDynamic* NewMaterial = CreateAndSetMaterialInstanceDynamicFromMaterial(MaterialID, CurrentMaterial);
 
-			FLinearColor ColorToApply;
+			FLinearColor ColorToApply = FLinearColor::Black;
 
 			switch (Coloring)
 			{
