Fixed some 'deprecated' FString usage.

diff --git a/Engine/Plugins/Enterprise/DatasmithImporter/Source/ExternalSource/Public/SourceUri.h b/Engine/Plugins/Enterprise/DatasmithImporter/Source/ExternalSource/Public/SourceUri.h
--- a/Engine/Plugins/Enterprise/DatasmithImporter/Source/ExternalSource/Public/SourceUri.h
+++ b/Engine/Plugins/Enterprise/DatasmithImporter/Source/ExternalSource/Public/SourceUri.h
@@ -103,7 +103,7 @@ namespace UE::DatasmithImporter
 			if (QueryContainer.Num() > 0)
 			{
 				TCHAR QueryDelimiter(GetQueryStartCharacter());
-				FString QueryString(1, &QueryDelimiter);
+				FString QueryString = FString::ConstructFromPtrSize(&QueryDelimiter, 1);
 				for (typename TContainer::TConstIterator It = QueryContainer.CreateConstIterator(); It;)
 				{
 					const TPair<FString, FString>& CurrentPair = *It;
diff --git a/Engine/Plugins/Experimental/DynamicMaterial/Source/DynamicMaterial/Private/DMComponentPath.cpp b/Engine/Plugins/Experimental/DynamicMaterial/Source/DynamicMaterial/Private/DMComponentPath.cpp
--- a/Engine/Plugins/Experimental/DynamicMaterial/Source/DynamicMaterial/Private/DMComponentPath.cpp
+++ b/Engine/Plugins/Experimental/DynamicMaterial/Source/DynamicMaterial/Private/DMComponentPath.cpp
@@ -126,6 +126,6 @@ bool FDMComponentPathSegment::GetParameter(FString& OutParameter) const
 		return false;
 	}
 
-	OutParameter = FString(Parameter.Len(), Parameter.GetData());
+	OutParameter = FString::ConstructFromPtrSize(Parameter.GetData(), Parameter.Len());
 	return true;
 }
diff --git a/Engine/Plugins/Experimental/SVGImporter/Source/SVGImporter/Private/SVGImporterUtils.cpp b/Engine/Plugins/Experimental/SVGImporter/Source/SVGImporter/Private/SVGImporterUtils.cpp
--- a/Engine/Plugins/Experimental/SVGImporter/Source/SVGImporter/Private/SVGImporterUtils.cpp
+++ b/Engine/Plugins/Experimental/SVGImporter/Source/SVGImporter/Private/SVGImporterUtils.cpp
@@ -343,7 +343,7 @@ TArray<FSVGStyle> FSVGImporterUtils::StylesFromCSS(FString InData)
 					if (ClassNameLength >= 0)
 					{
 						// todo: expand on this kind of syntax?  something.class
-						CurrClassName = FString(ClassNameLength, Start + CurrClassNameBeginIndex);
+						CurrClassName = FString::ConstructFromPtrSize(Start + CurrClassNameBeginIndex, ClassNameLength);
 						if (CurrClassName.StartsWith("."))
 						{
 							CurrClassName.RemoveAt(0);
@@ -382,10 +382,10 @@ TArray<FSVGStyle> FSVGImporterUtils::StylesFromCSS(FString InData)
 					const int32 AttributeValueLength = CurrAttributeValueEndIndex - CurrAttributeValueBeginIndex;
 					if (AttributeNameLength >= 0 || AttributeValueLength >= 0)
 					{
-						FString AttributeName = FString(AttributeNameLength, Start + CurrAttributeNameBeginIndex);
+						FString AttributeName = FString::ConstructFromPtrSize(Start + CurrAttributeNameBeginIndex, AttributeNameLength);
 						AttributeName.RemoveSpacesInline();
 
-						FString AttributeValue = FString(CurrAttributeValueEndIndex - CurrAttributeValueBeginIndex, Start + CurrAttributeValueBeginIndex);
+						FString AttributeValue = FString::ConstructFromPtrSize(Start + CurrAttributeValueBeginIndex, CurrAttributeValueEndIndex - CurrAttributeValueBeginIndex);
 						AttributeValue.RemoveSpacesInline();
 
 						CurrStyleAttributes.Add(AttributeName, AttributeValue);
diff --git a/Engine/Plugins/Experimental/WebSocketMessaging/Source/WebSocketMessaging/Private/WebSocketMessageTransport.cpp b/Engine/Plugins/Experimental/WebSocketMessaging/Source/WebSocketMessaging/Private/WebSocketMessageTransport.cpp
--- a/Engine/Plugins/Experimental/WebSocketMessaging/Source/WebSocketMessaging/Private/WebSocketMessageTransport.cpp
+++ b/Engine/Plugins/Experimental/WebSocketMessaging/Source/WebSocketMessaging/Private/WebSocketMessageTransport.cpp
@@ -198,7 +198,7 @@ void FWebSocketMessageTransport::OnJsonMessage(const FString& InMessage, FWebSoc
 
 void FWebSocketMessageTransport::OnServerJsonMessage(void* InData, int32 InDataSize, FWebSocketMessageConnectionRef InWebSocketMessageConnection)
 {
-	FString Message(InDataSize, reinterpret_cast<UTF8CHAR*>(InData));
+	FString Message = FString::ConstructFromPtrSize(reinterpret_cast<UTF8CHAR*>(InData), InDataSize);
 	OnJsonMessage(Message, InWebSocketMessageConnection);
 }
 
diff --git a/Engine/Plugins/Media/PixelStreaming/Source/PixelStreaming/Private/PixelStreamingSignallingConnection.cpp b/Engine/Plugins/Media/PixelStreaming/Source/PixelStreaming/Private/PixelStreamingSignallingConnection.cpp
--- a/Engine/Plugins/Media/PixelStreaming/Source/PixelStreaming/Private/PixelStreamingSignallingConnection.cpp
+++ b/Engine/Plugins/Media/PixelStreaming/Source/PixelStreaming/Private/PixelStreamingSignallingConnection.cpp
@@ -457,7 +457,7 @@ void FPixelStreamingSignallingConnection::OnBinaryMessage(const uint8* Data, int
 {
 	FUTF8ToTCHAR Convert((const ANSICHAR*)Data, Length);
 	const TCHAR* PayloadChars = Convert.Get();
-	FString Msg = FString(Convert.Length(), PayloadChars);
+	FString Msg = FString::ConstructFromPtrSize(PayloadChars, Convert.Length());
 	OnMessage(Msg);
 }
 
diff --git a/Engine/Plugins/Media/PixelStreaming/Source/PixelStreaming/Private/Streamer.cpp b/Engine/Plugins/Media/PixelStreaming/Source/PixelStreaming/Private/Streamer.cpp
--- a/Engine/Plugins/Media/PixelStreaming/Source/PixelStreaming/Private/Streamer.cpp
+++ b/Engine/Plugins/Media/PixelStreaming/Source/PixelStreaming/Private/Streamer.cpp
@@ -792,7 +792,7 @@ namespace UE::PixelStreaming
 				{
 					const size_t DescriptorSize = (RawBuffer.data.size() - 1) / sizeof(TCHAR);
 					const TCHAR* DescPtr = reinterpret_cast<const TCHAR*>(RawBuffer.data.data() + 1);
-					const FString Message(DescriptorSize, DescPtr);
+					const FString Message = FString::ConstructFromPtrSize(DescPtr, DescriptorSize);
 					PlayerContext->DataChannel->SendMessage(FPixelStreamingInputProtocol::FromStreamerProtocol.Find("TestEcho")->GetID(), Message);
 				}
 			}
diff --git a/Engine/Plugins/Media/PixelStreaming/Source/PixelStreaming/Private/Tests/DataChannelTest.cpp b/Engine/Plugins/Media/PixelStreaming/Source/PixelStreaming/Private/Tests/DataChannelTest.cpp
--- a/Engine/Plugins/Media/PixelStreaming/Source/PixelStreaming/Private/Tests/DataChannelTest.cpp
+++ b/Engine/Plugins/Media/PixelStreaming/Source/PixelStreaming/Private/Tests/DataChannelTest.cpp
@@ -44,7 +44,7 @@ namespace UE::PixelStreaming
 				*bComplete.Get() = true;
 				const size_t DescriptorSize = (RawBuffer.data.size() - 1) / sizeof(TCHAR);
 				const TCHAR* DescPtr = reinterpret_cast<const TCHAR*>(RawBuffer.data.data() + 1);
-				const FString Message(DescriptorSize, DescPtr);
+				const FString Message = FString::ConstructFromPtrSize(DescPtr, DescriptorSize);
 				TestTrue(FString::Printf(TEXT("Received echo (%s) != sent echo (%s)."), *Message, *EchoContent), Message == EchoContent);
 			}
 		};
diff --git a/Engine/Plugins/Media/PixelStreaming/Source/PixelStreaming/Private/Tests/ProtocolTest.cpp b/Engine/Plugins/Media/PixelStreaming/Source/PixelStreaming/Private/Tests/ProtocolTest.cpp
--- a/Engine/Plugins/Media/PixelStreaming/Source/PixelStreaming/Private/Tests/ProtocolTest.cpp
+++ b/Engine/Plugins/Media/PixelStreaming/Source/PixelStreaming/Private/Tests/ProtocolTest.cpp
@@ -45,7 +45,7 @@ namespace UE::PixelStreaming
 			{
 				const size_t DescriptorSize = (RawBuffer.data.size() - 1) / sizeof(TCHAR);
 				const TCHAR* DescPtr = reinterpret_cast<const TCHAR*>(RawBuffer.data.data() + 1);
-				const FString JsonRaw(DescriptorSize, DescPtr);
+				const FString JsonRaw = FString::ConstructFromPtrSize(DescPtr, DescriptorSize);
 
 				TSharedPtr<FJsonObject> JsonParsed;
 				TSharedRef<TJsonReader<TCHAR>> JsonReader = TJsonReaderFactory<TCHAR>::Create(JsonRaw);
diff --git a/Engine/Plugins/Media/PixelStreaming/Source/PixelStreaming/Private/WebRTCLogging.cpp b/Engine/Plugins/Media/PixelStreaming/Source/PixelStreaming/Private/WebRTCLogging.cpp
--- a/Engine/Plugins/Media/PixelStreaming/Source/PixelStreaming/Private/WebRTCLogging.cpp
+++ b/Engine/Plugins/Media/PixelStreaming/Source/PixelStreaming/Private/WebRTCLogging.cpp
@@ -63,7 +63,7 @@ private:
 		// but since logging should be multithread-safe, it should be thread-local member or explicitly synchronised
 		auto StrPtr = StringCast<TCHAR>(message.c_str(), Size);
 		// to zero-terminate the string
-		FString Msg{ StrPtr.Length(), StrPtr.Get() };
+		FString Msg = FString::ConstructFromPtrSize( StrPtr.Get(), StrPtr.Length() );
 
 		switch (severity)
 		{
diff --git a/Engine/Plugins/Online/OnlineSubsystemGoogle/Source/Private/OnlineIdentityGoogleCommon.cpp b/Engine/Plugins/Online/OnlineSubsystemGoogle/Source/Private/OnlineIdentityGoogleCommon.cpp
--- a/Engine/Plugins/Online/OnlineSubsystemGoogle/Source/Private/OnlineIdentityGoogleCommon.cpp
+++ b/Engine/Plugins/Online/OnlineSubsystemGoogle/Source/Private/OnlineIdentityGoogleCommon.cpp
@@ -423,7 +423,7 @@ FUniqueNetIdPtr FOnlineIdentityGoogleCommon::CreateUniquePlayerId(uint8* Bytes,
 {
 	if (Bytes != nullptr && Size > 0)
 	{
-		FString StrId(Size, (TCHAR*)Bytes);
+		FString StrId = FString::ConstructFromPtrSize((TCHAR*)Bytes, Size);
 		return FUniqueNetIdGoogle::Create(StrId);
 	}
 	return nullptr;
diff --git a/Engine/Plugins/Protocols/MQTT/Source/MQTTCore/Private/MQTTClient.cpp b/Engine/Plugins/Protocols/MQTT/Source/MQTTCore/Private/MQTTClient.cpp
--- a/Engine/Plugins/Protocols/MQTT/Source/MQTTCore/Private/MQTTClient.cpp
+++ b/Engine/Plugins/Protocols/MQTT/Source/MQTTCore/Private/MQTTClient.cpp
@@ -354,7 +354,7 @@ void FMQTTClient::OnMessagePacket(const FMQTTPublishPacket& InPacket)
 	Message.Payload = InPacket.Payload;
 	Message.QoS = InPacket.QoS;
 	Message.bRetain = InPacket.bRetain;
-	Message.PayloadString = FString(Message.Payload.Num(), (const char*)Message.Payload.GetData());
+	Message.PayloadString = FString::ConstructFromPtrSize((const char*)Message.Payload.GetData(), Message.Payload.Num());
 
 	// No subscriptions, can happen if that overload isn't used
 	if(Subscriptions.Num() == 0)
diff --git a/Engine/Plugins/Runtime/OSC/Source/OSC/Private/OSCStream.cpp b/Engine/Plugins/Runtime/OSC/Source/OSC/Private/OSCStream.cpp
--- a/Engine/Plugins/Runtime/OSC/Source/OSC/Private/OSCStream.cpp
+++ b/Engine/Plugins/Runtime/OSC/Source/OSC/Private/OSCStream.cpp
@@ -167,7 +167,7 @@ namespace UE::OSC
 		}
 
 		// Exclude null terminator here; this constructor appends one.
-		return FString(Count - 1, (ANSICHAR*)(&Data[InitPosition]));
+		return FString::ConstructFromPtrSize((ANSICHAR*)(&Data[InitPosition]), Count - 1);
 	}
 
 	void FStream::WriteString(const FString& InString)
