MotionDesign : SVGImporter
- Fix crash caused by incorrect parsing of points path in svg data

diff --git a/Engine/Plugins/Experimental/SVGImporter/Source/SVGImporterEditor/Private/SVGParsingUtils.cpp b/Engine/Plugins/Experimental/SVGImporter/Source/SVGImporterEditor/Private/SVGParsingUtils.cpp
--- a/Engine/Plugins/Experimental/SVGImporter/Source/SVGImporterEditor/Private/SVGParsingUtils.cpp
+++ b/Engine/Plugins/Experimental/SVGImporter/Source/SVGImporterEditor/Private/SVGParsingUtils.cpp
@@ -67,38 +67,21 @@ int32 ParseDecimals(TArray<FString>& OutArray, const FString& Data, bool InCullE
 void FSVGParsingUtils::PointsFromString(FString InPointsString, TArray<FVector2D>& OutPoints)
 {
 	OutPoints.Empty();
-	TArray<FString> Values;
 
-	InPointsString.ReplaceInline(TEXT("-"), TEXT(" -")); // we could have points written like this: 229-303 0-66, messing up the parsing
-	InPointsString.ParseIntoArrayWS(Values, TEXT(","), true); // parses both based on a WhiteSpace and "," delimiters
+	constexpr bool bApplyScale = true;
+	const TArray<float> Points = FloatsFromString(InPointsString, bApplyScale);
 
-	TArray<float> OutFloats;
-
-	for (const FString& ValueString : Values)
+	if (Points.Num() % 2 != 0)
 	{
-		TArray<FString> Floats;
-
-		ParseDecimals(Floats, ValueString, true); // we could have points written like this: 31.56.12, meaning 31.56 0.12
-
-		for (FString& CurrVal : Floats)
-		{
-			float Val = FCString::Atof(*CurrVal);
-			OutFloats.Add(Val);
-		}
-	}
-
-	if (OutFloats.Num() % 2 != 0)
-	{
-		UE_LOG(SVGImporterEditorLog, Warning, TEXT("Trying to parse SVG attribute as point list, but number of resulting members is not even. Returned points list will be empty."));
+		UE_LOG(SVGImporterEditorLog, Warning, TEXT("Trying to parse SVG attribute as point list, but number of resulting members is not even %i. Returned points list will be empty. %s"), Points.Num(), *InPointsString);
 		return;
 	}
 
-	for (int32 i = 0; i < OutFloats.Num(); i+=2)
-	{
-		float X = OutFloats[i];
-		float Y = OutFloats[i+1];
-		OutPoints.Add(FVector2D(X, Y) * UE::SVGImporterEditor::Private::SVGScaleFactor);
-	}
+	OutPoints.Reserve(Points.Num() / 2);
+    for (int32 PointIndex = 0; PointIndex + 1 < Points.Num(); PointIndex += 2)
+    {
+        OutPoints.Add(FVector2D(Points[PointIndex], Points[PointIndex + 1]));
+    }
 }
 
 TSharedRef<FSVGParser_Base> FSVGParsingUtils::CreateSVGParser(const FString& InStringToParse, ESVGParserType InParserType)
@@ -144,32 +127,88 @@ bool FSVGParsingUtils::FloatFromStringWithSuffix(const FString& InFloatString, c
 
 TArray<float> FSVGParsingUtils::FloatsFromString(FString InFloatsString, bool bApplyScaling /*default == true*/)
 {
-	TArray<FString> Values;
-	InFloatsString.ReplaceInline(TEXT("-"), TEXT(" -")); // we could have points written like this: 229-303 0-66, messing up the parsing
-	InFloatsString.ParseIntoArrayWS(Values, TEXT(","), true); // parses both based on a WhiteSpace and "," delimiters
-
-	TArray<float> OutFloats;
-
-	for (const FString& ValueString : Values)
-	{
-		TArray<FString> Floats;
-
-		ParseDecimals(Floats, ValueString, true);
-
-		for (FString& CurrVal : Floats)
-		{
-			float Val = FCString::Atof(*CurrVal);
-
-			if (bApplyScaling)
-			{
-				Val*=UE::SVGImporterEditor::Private::SVGScaleFactor;
-			}
-
-			OutFloats.Add(Val);
-		}
-	}
-
-	return OutFloats;
+	InFloatsString.TrimStartAndEndInline();
+	
+    TArray<float> ParsedValues;
+    FString CurrentNumber;
+    bool bIsScientificNotation = false;
+	bool bHasDecimalPoint = false;
+	const float Multiplier = bApplyScaling ? UE::SVGImporterEditor::Private::SVGScaleFactor : 1.f;
+
+	/**
+	 * Formats supported
+	 * 10,20 => (10, 20)
+	 * 10 20 => (10, 20)
+	 * -3.48-1.92-.53.66 => (-3.48,-1.92),(-0.53,0.66)
+	 * .62.78,3.76-2.11 => (0.62,0.78),(3.76,-2.11)
+	 * .10, .20 => (0.10,0.20)
+	 * .10 .20 => (0.10,0.20)
+	 * .62e3.78E-1 => (0.62e3,0.78e-1)
+	 */
+    for (int32 CharIndex = 0; CharIndex < InFloatsString.Len(); ++CharIndex)
+    {
+	    const TCHAR Char = InFloatsString[CharIndex];
+
+        // Check if character is part of a number : -.10
+        if (FChar::IsDigit(Char) || (Char == '.' && !bHasDecimalPoint) || (CurrentNumber.IsEmpty() && (Char == '-' || Char == '+')))
+        {
+            CurrentNumber += Char;
+
+        	if (Char == '.')
+        	{
+        		bHasDecimalPoint = true;
+        	}
+        }
+    	// Scientific notation detected; mark the flag and include 'e'/'E' : -10e3
+        else if ((Char == 'e' || Char == 'E') && !CurrentNumber.IsEmpty())
+        {
+            bIsScientificNotation = true;
+            CurrentNumber += Char;
+        }
+    	// Scientific notation exponent sign : -10e3
+        else if ((Char == '-' || Char == '+') && bIsScientificNotation)
+        {
+            CurrentNumber += Char;
+            bIsScientificNotation = false;
+        }
+    	// non-numeric character encountered
+        else
+        {
+            if (!CurrentNumber.IsEmpty())
+            {
+                ParsedValues.Add(FCString::Atof(*CurrentNumber) * Multiplier);
+            }
+
+        	// Reset
+        	CurrentNumber.Empty();
+            bIsScientificNotation = false;
+        	bHasDecimalPoint = false;
+
+            // Skip comma as a delimiter
+            if (Char == ',')
+            {
+                continue;
+            }
+
+            // Start next number
+            if (Char == '-' || Char == '+' || Char == '.')
+            {
+                CurrentNumber += Char;
+            	
+            	if (Char == '.')
+            	{
+            		bHasDecimalPoint = true;
+            	}
+            }
+        }
+    }
+
+    if (!CurrentNumber.IsEmpty())
+    {
+        ParsedValues.Add(FCString::Atof(*CurrentNumber) * Multiplier);
+    }
+
+	return ParsedValues;
 }
 
 int32 FSVGParsingUtils::ParseIntoCommandsArrayKeepDelimiters(const FString& InSourceString, TArray<FString>& OutArray, const TCHAR* const * DelimArray, int32 NumDelims, bool InCullEmpty)
@@ -309,23 +348,29 @@ TArray<TArray<FSVGPathCommand>> FSVGParsingUtils::ParseStringAsPathCommands(cons
 		TCHAR CommandType = CommandString.GetCharArray()[0];
 		CommandString.RemoveAt(0);
 
-		TArray<FVector2D> Points;
+		bool bInvalidCommand = false;
 
 		switch (CommandType)
 		{
 			case 'm':
 			case 'M':
 				{
+					TArray<FVector2D> MPoints;
+					PointsFromString(CommandString, MPoints);
+					const bool bIsRelative = CommandType == 'm';
+
+					if (MPoints.IsEmpty())
+					{
+						bInvalidCommand = true;
+						break;
+					}
+
 					// moveto parameters: (x y)+
 					// move to implies the start of a new sub path, let's create it
 					TArray<FSVGPathCommand> NewSubPath;
 					SubPaths.Add(NewSubPath);
 					CurrPathIndex++;
 
-					TArray<FVector2D> MPoints;
-					PointsFromString(CommandString, MPoints);
-					const bool bIsRelative = CommandType == 'm';
-
 					// We store the initial point of the sub path to properly apply closepath later, if needed
 					InitialPoint = PathMoveTo(SubPaths[CurrPathIndex], CursorPos, MPoints, bIsRelative);
 					CursorPos2 = InitialPoint;
@@ -335,6 +380,12 @@ TArray<TArray<FSVGPathCommand>> FSVGParsingUtils::ParseStringAsPathCommands(cons
 			case 'z':
 			case 'Z':
 				{
+					if (SubPaths.IsEmpty())
+					{
+						bInvalidCommand = true;
+						break;
+					}
+					
 					// closepath parameters: none
 					FSVGPathCommand ClosePathCmd(ESVGPathInstructionType::ClosePath);
 					ClosePathCmd.PointTo = InitialPoint;
@@ -352,6 +403,12 @@ TArray<TArray<FSVGPathCommand>> FSVGParsingUtils::ParseStringAsPathCommands(cons
 					PointsFromString(CommandString, LPoints);
 					const bool bIsRelative = CommandType == 'l';
 
+					if (LPoints.IsEmpty())
+					{
+						bInvalidCommand = true;
+						break;
+					}
+
 					PathLineTo(SubPaths[CurrPathIndex], CursorPos, LPoints, bIsRelative);
 					CursorPos2 = CursorPos;
 				}
@@ -364,6 +421,12 @@ TArray<TArray<FSVGPathCommand>> FSVGParsingUtils::ParseStringAsPathCommands(cons
 					TArray<float> HFloats = FloatsFromString(CommandString);
 					const bool bIsRelative = CommandType == 'h';
 
+					if (HFloats.IsEmpty())
+					{
+						bInvalidCommand = true;
+						break;
+					}
+
 					PathHorizontalLineTo(SubPaths[CurrPathIndex], CursorPos, HFloats, bIsRelative);
 					CursorPos2 = CursorPos;
 				}
@@ -376,6 +439,12 @@ TArray<TArray<FSVGPathCommand>> FSVGParsingUtils::ParseStringAsPathCommands(cons
 					TArray<float> VFloats = FloatsFromString(CommandString);
 					const bool bIsRelative = CommandType == 'v';
 
+					if (VFloats.IsEmpty())
+					{
+						bInvalidCommand = true;
+						break;
+					}
+					
 					PathVerticalLineTo(SubPaths[CurrPathIndex], CursorPos, VFloats, bIsRelative);
 					CursorPos2 = CursorPos;
 				}
@@ -389,6 +458,12 @@ TArray<TArray<FSVGPathCommand>> FSVGParsingUtils::ParseStringAsPathCommands(cons
 					PointsFromString(CommandString, CPoints);
 					const bool bIsRelative = CommandType == 'c';
 
+					if (CPoints.IsEmpty())
+					{
+						bInvalidCommand = true;
+						break;
+					}
+
 					PathCubicBezierTo(SubPaths[CurrPathIndex], CursorPos, CursorPos2, CPoints, bIsRelative);
 				}
 				break;
@@ -401,6 +476,12 @@ TArray<TArray<FSVGPathCommand>> FSVGParsingUtils::ParseStringAsPathCommands(cons
 					PointsFromString(CommandString, SPoints);
 					const bool bIsRelative = CommandType == 's';
 
+					if (SPoints.IsEmpty())
+					{
+						bInvalidCommand = true;
+						break;
+					}
+
 					PathCubicBezierSmoothTo(SubPaths[CurrPathIndex], CursorPos, CursorPos2, SPoints, bIsRelative);
 				}
 				break;
@@ -412,6 +493,12 @@ TArray<TArray<FSVGPathCommand>> FSVGParsingUtils::ParseStringAsPathCommands(cons
 					TArray<FVector2D> QPoints;
 					PointsFromString(CommandString, QPoints);
 					const bool bIsRelative = CommandType == 'q';
+					
+					if (QPoints.IsEmpty())
+					{
+						bInvalidCommand = true;
+						break;
+					}
 
 					PathQuadraticBezierTo(SubPaths[CurrPathIndex], CursorPos, CursorPos2, QPoints, bIsRelative);
 				}
@@ -425,6 +512,12 @@ TArray<TArray<FSVGPathCommand>> FSVGParsingUtils::ParseStringAsPathCommands(cons
 					PointsFromString(CommandString, TPoints);
 					const bool bIsRelative = CommandType == 't';
 
+					if (TPoints.IsEmpty())
+					{
+						bInvalidCommand = true;
+						break;
+					}
+
 					PathQuadraticBezierSmoothTo(SubPaths[CurrPathIndex], CursorPos, CursorPos2, TPoints, bIsRelative);
 				}
 				break;
@@ -437,6 +530,12 @@ TArray<TArray<FSVGPathCommand>> FSVGParsingUtils::ParseStringAsPathCommands(cons
 					TArray<float> AValues = FloatsFromString(CommandString, bApplyScaling);
 					bool bIsRelative = CommandType == 'a';
 
+					if (AValues.IsEmpty())
+					{
+						bInvalidCommand = true;
+						break;
+					}
+					
 					PathArcTo(SubPaths[CurrPathIndex], CursorPos, AValues, bIsRelative);
 					CursorPos2 = CursorPos;
 				}
@@ -445,6 +544,13 @@ TArray<TArray<FSVGPathCommand>> FSVGParsingUtils::ParseStringAsPathCommands(cons
 			default:
 			break;
 		}
+
+		if (bInvalidCommand)
+		{
+			UE_LOG(SVGImporterEditorLog, Warning, TEXT("Invalid svg command detected during parsing (%s), skipping entire svg subpath (%s)"), *CommandStringElem, *InString);
+			SubPaths.Empty();
+			break;
+		}
 	}
 
 	return SubPaths;
