MotionDesign : SVG
- Fixed duplicated joined SVG actor would not apply colors to materials
- Fixed duplicated joined SVG actor would not regenerate geometry when render mode is changed (2D -> 3D)
- Fixed pixelated SVG materials, changed materials to translucent instead of masked and set sort priority for ordering
- Add ability to use SVG as target for mask instead of only source

#virtualized

diff --git a/Engine/Plugins/Experimental/SVGImporter/Source/SVGImporter/Private/ProceduralMeshes/JoinedSVGDynamicMeshComponent.cpp b/Engine/Plugins/Experimental/SVGImporter/Source/SVGImporter/Private/ProceduralMeshes/JoinedSVGDynamicMeshComponent.cpp
--- a/Engine/Plugins/Experimental/SVGImporter/Source/SVGImporter/Private/ProceduralMeshes/JoinedSVGDynamicMeshComponent.cpp
+++ b/Engine/Plugins/Experimental/SVGImporter/Source/SVGImporter/Private/ProceduralMeshes/JoinedSVGDynamicMeshComponent.cpp
@@ -12,35 +12,6 @@
 #include "SVGImporterUtils.h"
 #include "UObject/ConstructorHelpers.h"
 
-bool FSVGShapeParameters::ImportTextItem(const TCHAR*& Buffer, int32 PortFlags, UObject* Parent, FOutputDevice* ErrorText)
-{
-	FString ImportedInstance;
-	const TCHAR* NewBuffer = FPropertyHelpers::ReadToken(Buffer, ImportedInstance, true);
-
-	if (!NewBuffer)
-	{
-		// Failed to read buffer
-		return false;
-	}
-
-	if (ImportedInstance[0] == '(')
-	{
-		// We can use normal ImportText to parse the Buffer, and then just copy what we actually want to copy
-		const UScriptStruct* ScriptStruct = FSVGShapeParameters::StaticStruct();
-		FSVGShapeParameters StructCopy;
-
-		if (ScriptStruct->ImportText(Buffer, &StructCopy, Parent, PortFlags, ErrorText, ScriptStruct->GetName(), false))
-		{
-			// We are not import/paste the shape name or the Material ID, since the expected workflow is to just handle the Color
-			Color = StructCopy.Color;
-		}
-
-		return true;
-	}
-
-	return false;
-}
-
 UJoinedSVGDynamicMeshComponent::UJoinedSVGDynamicMeshComponent()
 {
 	bSVGIsUnlit = true;
diff --git a/Engine/Plugins/Experimental/SVGImporter/Source/SVGImporter/Private/SVGActor.cpp b/Engine/Plugins/Experimental/SVGImporter/Source/SVGImporter/Private/SVGActor.cpp
--- a/Engine/Plugins/Experimental/SVGImporter/Source/SVGImporter/Private/SVGActor.cpp
+++ b/Engine/Plugins/Experimental/SVGImporter/Source/SVGImporter/Private/SVGActor.cpp
@@ -280,6 +280,11 @@ void ASVGActor::PostLoad()
 	{
 		Generate();
 	}
+
+	if (RenderMode == ESVGRenderMode::DynamicMesh3D)
+	{
+		RefreshAllShapesMaterials();
+	}
 }
 
 void ASVGActor::PostDuplicate(bool bDuplicateForPIE)
@@ -466,7 +471,6 @@ void ASVGActor::CreateMeshesFromShape(const FSVGShape& InShape)
 {
 	if (!SVGData)
 	{
-		bIsGeneratingMeshes = false;
 #if WITH_EDITOR
 		DisplayMissingSVGDataError(TEXT("Cannot generate SVG geometry."));
 #endif
@@ -740,10 +744,7 @@ void ASVGActor::ScheduleShapesGeneration()
 		CurrExtrudeForDepth = 0.0f;
 		GeneratedShapesCount = 0;
 
-		bIsGeneratingMeshes = true;
 		GenMeshDelegateHandle = FCoreDelegates::OnEndFrame.AddUObject(this, &ASVGActor::GenerateShapesOnEndFrame);
-
-		bMeshesShouldBeGenerated = false;
 		OnShapesGenerationEnd_Delegate.AddUObject(this, &ASVGActor::OnShapesGenerationEnd);
 	}
 }
@@ -847,7 +848,6 @@ void ASVGActor::GenerateNextMesh()
 {
 	if (!SVGData)
 	{
-		bIsGeneratingMeshes = false;
 		FCoreDelegates::OnEndFrame.Remove(GenMeshDelegateHandle);
 
 #if WITH_EDITOR
@@ -949,9 +949,15 @@ void ASVGActor::ApplyScaleAndCenter()
 
 void ASVGActor::GenerateShapesOnEndFrame()
 {
+	if (bIsGeneratingMeshes || !bMeshesShouldBeGenerated)
+	{
+		return;
+	}
+
+	TGuardValue<bool> GeneratingGuard(bIsGeneratingMeshes, true);
+
 	if (!SVGData)
 	{
-		bIsGeneratingMeshes = false;
 		FCoreDelegates::OnEndFrame.Remove(GenMeshDelegateHandle);
 
 #if WITH_EDITOR
@@ -977,6 +983,7 @@ void ASVGActor::GenerateShapesOnEndFrame()
 			if (!SVGData->Shapes.IsValidIndex(GeneratedShapesCount))
 			{
 				bHasFinished = true;
+				bMeshesShouldBeGenerated = false;
 				FCoreDelegates::OnEndFrame.Remove(GenMeshDelegateHandle);
 
 				ScheduleFillExtrudeUpdate();
@@ -985,9 +992,6 @@ void ASVGActor::GenerateShapesOnEndFrame()
 				// notify details panel, so that newly added components are shown right away
 				TriggerActorDetailsRefresh();
 
-				bMeshesShouldBeGenerated = false;
-				bIsGeneratingMeshes = false;
-
 				bSVGHasFills = !GetFillComponents().IsEmpty();
 				bSVGHasStrokes = !GetStrokeComponents().IsEmpty();
 
@@ -1038,6 +1042,11 @@ void ASVGActor::OnShapesGenerationEnd()
 	RefreshFillShapesExtrude(EPropertyChangeType::ValueSet);
 	RefreshStrokeShapesExtrude(EPropertyChangeType::ValueSet);
 	ApplyOffset();
+
+	if (RenderMode == ESVGRenderMode::DynamicMesh3D)
+	{
+		RefreshAllShapesMaterials();
+	}
 }
 
 void ASVGActor::OnShapesExtrudeEnd()
@@ -1285,11 +1294,13 @@ void ASVGActor::RefreshAllShapes()
 
 void ASVGActor::RefreshAllShapesMaterials()
 {
+	int32 SortPriority = 0;
 	for (USVGDynamicMeshComponent* ShapeComponent : ShapeComponents)
 	{
 		if (ShapeComponent)
 		{
 			ShapeComponent->SetIsUnlit(bSVGIsUnlit);
+			ShapeComponent->SetTranslucentSortPriority(SortPriority++);
 		}
 	}
 }
diff --git a/Engine/Plugins/Experimental/SVGImporter/Source/SVGImporter/Public/ProceduralMeshes/JoinedSVGDynamicMeshComponent.h b/Engine/Plugins/Experimental/SVGImporter/Source/SVGImporter/Public/ProceduralMeshes/JoinedSVGDynamicMeshComponent.h
--- a/Engine/Plugins/Experimental/SVGImporter/Source/SVGImporter/Public/ProceduralMeshes/JoinedSVGDynamicMeshComponent.h
+++ b/Engine/Plugins/Experimental/SVGImporter/Source/SVGImporter/Public/ProceduralMeshes/JoinedSVGDynamicMeshComponent.h
@@ -22,8 +22,6 @@ struct FSVGShapeParameters
 	UPROPERTY()
 	int32 MaterialID = 0;
 
-	SVGIMPORTER_API bool ImportTextItem(const TCHAR*& Buffer, int32 PortFlags, UObject* Parent, FOutputDevice* ErrorText);
-
 	bool operator==(const FSVGShapeParameters& InShapeParameters) const
 	{
 		return ShapeName == InShapeParameters.ShapeName
@@ -37,15 +35,6 @@ struct FSVGShapeParameters
 	}
 };
 
-template<>
-struct TStructOpsTypeTraits<FSVGShapeParameters> : TStructOpsTypeTraitsBase2<FSVGShapeParameters>
-{
-	enum
-	{
-		WithImportTextItem = true,
-	};
-};
-
 USTRUCT()
 struct FJoinedSVGMeshParameters
 {
