MotionDesign : SVGImporter
- Added option to toggle visibility of all svg components in the actor
- Switched materials to masked instead of translucent to avoid inverted normals

#virtualized

diff --git a/Engine/Plugins/Experimental/SVGImporter/Source/SVGImporter/Private/SVGActor.cpp b/Engine/Plugins/Experimental/SVGImporter/Source/SVGImporter/Private/SVGActor.cpp
--- a/Engine/Plugins/Experimental/SVGImporter/Source/SVGImporter/Private/SVGActor.cpp
+++ b/Engine/Plugins/Experimental/SVGImporter/Source/SVGImporter/Private/SVGActor.cpp
@@ -370,9 +370,9 @@ void ASVGActor::AddStrokeComponent(const TArray<FVector>& InPoints, float InThic
 
 		USVGStrokeComponent* DynStrokeComponent = NewObject<USVGStrokeComponent>(this, StrokeShapeUniqueName, RF_Transactional);
 		AddInstanceComponent(DynStrokeComponent);
-		DynStrokeComponent->OnComponentCreated();
-		DynStrokeComponent->SetupAttachment(StrokeShapesRoot);
 		DynStrokeComponent->RegisterComponent();
+		DynStrokeComponent->AttachToComponent(StrokeShapesRoot, FAttachmentTransformRules::KeepRelativeTransform);
+		
 		ShapeComponents.Add(DynStrokeComponent);
 
 		DynStrokeComponent->SetExtrudeType(ExtrudeType);
@@ -413,9 +413,9 @@ void ASVGActor::AddFillComponent(const TArray<FSVGPathPolygon>& InShapesToDraw,
 
 		USVGFillComponent* DynFillComponent = NewObject<USVGFillComponent>(this, FillShapeUniqueName, RF_Transactional);
 		AddInstanceComponent(DynFillComponent);
-		DynFillComponent->OnComponentCreated();
-		DynFillComponent->SetupAttachment(FillShapesRoot);
 		DynFillComponent->RegisterComponent();
+		DynFillComponent->AttachToComponent(FillShapesRoot, FAttachmentTransformRules::KeepRelativeTransform);
+		
 		ShapeComponents.Add(DynFillComponent);
 
 		DynFillComponent->SetExtrudeType(ExtrudeType);
@@ -947,6 +947,18 @@ void ASVGActor::ApplyScaleAndCenter()
 	CenterMeshes();
 }
 
+void ASVGActor::OnVisibilityChanged()
+{
+	if (bVisible)
+	{
+		ShowShapes();
+	}
+	else
+	{
+		HideShapes();
+	}
+}
+
 void ASVGActor::GenerateShapesOnEndFrame()
 {
 	if (bIsGeneratingMeshes || !bMeshesShouldBeGenerated)
@@ -1016,7 +1028,7 @@ void ASVGActor::ShowShapes() const
 {
 	if (RootComponent)
 	{
-		RootComponent->SetVisibility(true, true);
+		RootComponent->SetVisibility(bVisible, true);
 	}
 }
 
@@ -1246,6 +1258,17 @@ void ASVGActor::SetStrokesWidth(float InStrokesWidth)
 	}
 }
 
+void ASVGActor::SetVisible(bool bInVisible)
+{
+	if (bVisible == bInVisible)
+	{
+		return;
+	}
+
+	bVisible = bInVisible;
+	OnVisibilityChanged();
+}
+
 void ASVGActor::RefreshFillShapesExtrude(const EPropertyChangeType::Type ChangeType)
 {
 	const float NewExtrudeValue = ExtrudeType == ESVGExtrudeType::None ? 0.0f : FillsExtrude;
@@ -1487,6 +1510,7 @@ void ASVGActor::PostEditChangeProperty(FPropertyChangedEvent& InPropertyChangedE
 	static const FName SVGIsUnlitName = GET_MEMBER_NAME_CHECKED(ASVGActor, bSVGIsUnlit);
 	static const FName SVGCastsShadowName = GET_MEMBER_NAME_CHECKED(ASVGActor, bSVGCastsShadow);
 	static const FName ScaleName = GET_MEMBER_NAME_CHECKED(ASVGActor, Scale);
+	static const FName VisibleName = GET_MEMBER_NAME_CHECKED(ASVGActor, bVisible);
 
 	if (MemberName == FillsExtrudeName)
 	{
@@ -1586,6 +1610,10 @@ void ASVGActor::PostEditChangeProperty(FPropertyChangedEvent& InPropertyChangedE
 			RefreshSVGPlaneCastShadow();
 		}
 	}
+	else if (MemberName == VisibleName)
+	{
+		OnVisibilityChanged();
+	}
 }
 
 void ASVGActor::PostTransacted(const FTransactionObjectEvent& InTransactionEvent)
diff --git a/Engine/Plugins/Experimental/SVGImporter/Source/SVGImporter/Public/SVGActor.h b/Engine/Plugins/Experimental/SVGImporter/Source/SVGImporter/Public/SVGActor.h
--- a/Engine/Plugins/Experimental/SVGImporter/Source/SVGImporter/Public/SVGActor.h
+++ b/Engine/Plugins/Experimental/SVGImporter/Source/SVGImporter/Public/SVGActor.h
@@ -97,6 +97,9 @@ public:
 	float GetStrokesWidth() const { return StrokesWidth; }
 	void SetStrokesWidth(float InStrokesWidth);
 
+	bool GetVisible() const { return bVisible; }
+	void SetVisible(bool bInVisible);
+
 protected:
 	void RefreshFillShapesExtrude(const EPropertyChangeType::Type ChangeType = EPropertyChangeType::ValueSet);
 	void RefreshStrokeShapesExtrude(const EPropertyChangeType::Type ChangeType = EPropertyChangeType::ValueSet);
@@ -180,6 +183,7 @@ protected:
 	void ApplyOffset();
 	void ApplyScaleAndCenter();
 
+	void OnVisibilityChanged();
 	void ShowShapes() const;
 	void HideShapes() const;
 
@@ -266,6 +270,9 @@ public:
 	bool bSVGCastsShadow;
 
 protected:
+	UPROPERTY(EditAnywhere, Setter="SetVisible", Getter="GetVisible", Category="Rendering")
+	bool bVisible = true;
+
 	UPROPERTY(EditDefaultsOnly, Category="SVG")
 	bool bSVGHasStrokes;
 
